Animation Locomotion Library: Common techniques for driving locomotion based character animation.

Included features in this changelist:
 * Capture a snapshot of common movement properties (velocity, speed, etc) that are used to drive animation
 * Character movement component for stop/pivot location
 * Library of anim node functions for Distance Matching
 * Template anim node that can advance by distance traveled rather than by time
 * Template anim node that's driven by distance to a target rather than by time
 * Animation modifier that generates distance curves based on root motion
 * Turn in place functionality for keeping the capsule from spinning the pose when it rotates by applying an offset, includes functionality for turn on spot animations to compenstate for the offset

diff --git a/Engine/Plugins/Animation/AnimationLocomotionLibrary/AnimationLocomotionLibrary.uplugin b/Engine/Plugins/Animation/AnimationLocomotionLibrary/AnimationLocomotionLibrary.uplugin
new file mode 100644
--- /dev/null
+++ b/Engine/Plugins/Animation/AnimationLocomotionLibrary/AnimationLocomotionLibrary.uplugin
@@ -0,0 +1,30 @@
+{
+	"FileVersion": 1,
+	"Version": 1,
+	"VersionName": "1.0",
+	"FriendlyName": "Animation Locomotion Library",
+	"Description": "Collection of techniques for driving locomotion animations",
+	"Category": "Animation",
+	"CreatedBy": "Epic Games, Inc.",
+	"CreatedByURL": "http://epicgames.com",
+	"DocsURL": "",
+	"MarketplaceURL": "",
+	"SupportURL": "",
+	"EnabledByDefault": false,
+	"CanContainContent": true,
+	"IsBetaVersion": true,
+	"IsExperimentalVersion": false,
+	"Installed": false,
+	"Modules": [
+		{
+			"Name": "AnimationLocomotionLibraryRuntime",
+			"Type": "Runtime",
+			"LoadingPhase": "Default"
+		},
+		{
+			"Name": "AnimationLocomotionLibraryEditor",
+			"Type": "Editor",
+			"LoadingPhase": "Default"
+		}
+	]
+}
\ No newline at end of file
diff --git a/Engine/Plugins/Animation/AnimationLocomotionLibrary/Source/Editor/AnimationLocomotionLibraryEditor.Build.cs b/Engine/Plugins/Animation/AnimationLocomotionLibrary/Source/Editor/AnimationLocomotionLibraryEditor.Build.cs
new file mode 100644
--- /dev/null
+++ b/Engine/Plugins/Animation/AnimationLocomotionLibrary/Source/Editor/AnimationLocomotionLibraryEditor.Build.cs
@@ -0,0 +1,31 @@
+// Copyright Epic Games, Inc. All Rights Reserved.
+
+namespace UnrealBuildTool.Rules
+{
+	public class AnimationLocomotionLibraryEditor : ModuleRules
+	{
+		public AnimationLocomotionLibraryEditor(ReadOnlyTargetRules Target) : base(Target)
+		{
+			PublicDependencyModuleNames.AddRange(
+			new string[]
+			{
+					"Core",
+					"CoreUObject",
+					"Engine",
+					"AnimGraphRuntime",
+					"AnimationModifiers",
+					"AnimationBlueprintLibrary"
+			});
+
+			PrivateDependencyModuleNames.AddRange(
+			new string[]
+			{
+            });
+
+			PrivateIncludePathModuleNames.AddRange(
+			new string[]
+			{
+			});
+		}
+	}
+}
diff --git a/Engine/Plugins/Animation/AnimationLocomotionLibrary/Source/Editor/Private/AnimationLocomotionLibraryEditor.cpp b/Engine/Plugins/Animation/AnimationLocomotionLibrary/Source/Editor/Private/AnimationLocomotionLibraryEditor.cpp
new file mode 100644
--- /dev/null
+++ b/Engine/Plugins/Animation/AnimationLocomotionLibrary/Source/Editor/Private/AnimationLocomotionLibraryEditor.cpp
@@ -0,0 +1,11 @@
+// Copyright Epic Games, Inc. All Rights Reserved.
+
+#include "AnimationLocomotionLibraryEditor.h"
+#include "Modules/ModuleManager.h"
+
+class FAnimationLocomotionLibraryEditorModule : public IAnimationLocomotionLibraryEditorModule
+{
+};
+
+
+IMPLEMENT_MODULE( FAnimationLocomotionLibraryEditorModule, AnimationLocomotionLibraryEditor )
diff --git a/Engine/Plugins/Animation/AnimationLocomotionLibrary/Source/Editor/Private/DistanceCurveModifier.cpp b/Engine/Plugins/Animation/AnimationLocomotionLibrary/Source/Editor/Private/DistanceCurveModifier.cpp
new file mode 100644
--- /dev/null
+++ b/Engine/Plugins/Animation/AnimationLocomotionLibrary/Source/Editor/Private/DistanceCurveModifier.cpp
@@ -0,0 +1,72 @@
+// Copyright Epic Games, Inc. All Rights Reserved.
+
+#include "DistanceCurveModifier.h"
+#include "AnimationBlueprintLibrary.h"
+#include "Animation/AnimSequence.h"
+#include "Animation/AnimationPoseData.h"
+
+// TODO: This logic works decently for simple clips but it should be reworked to be more robust:
+//  * It could detect pivot points by change in direction.
+//  * It should also account for clips that have multiple stop/pivot points.
+//  * It should handle distance traveled for the ends of looping animations.
+void UDistanceCurveModifier::OnApply_Implementation(UAnimSequence* Animation)
+{
+	if (Animation == nullptr)
+	{
+		UE_LOG(LogAnimation, Error, TEXT("DistanceCurveModifier failed. Reason: Invalid Animation"));
+		return;
+	}
+
+	if (!Animation->HasRootMotion())
+	{
+		UE_LOG(LogAnimation, Error, TEXT("DistanceCurveModifier failed. Reason: Root motion is disabled on the animation (%s)"), *GetNameSafe(Animation));
+		return;
+	}
+
+	const bool bMetaDataCurve = false;
+	UAnimationBlueprintLibrary::AddCurve(Animation, CurveName, ERawCurveTrackTypes::RCT_Float, bMetaDataCurve);
+
+	const float AnimLength = Animation->GetPlayLength();
+
+	// Perform a high resolution search to find the sample point with minimum speed.
+	
+	float TimeOfMinSpeed = 0.f;
+	float MinSpeedSq = FMath::Square(StopSpeedThreshold);
+
+	float SampleInterval = 1.f / 120.f;
+	int32 NumSteps = AnimLength / SampleInterval;
+	for (int32 Step = 0; Step < NumSteps; ++Step)
+	{
+		const float Time = Step * SampleInterval;
+
+		const bool bAllowLooping = false;
+		const FVector RootMotionTranslation = Animation->ExtractRootMotion(Time, SampleInterval, bAllowLooping).GetTranslation();
+		const float RootMotionSpeedSq = RootMotionTranslation.SizeSquared2D() / SampleInterval;
+
+		if (RootMotionSpeedSq < MinSpeedSq)
+		{
+			MinSpeedSq = RootMotionSpeedSq;
+			TimeOfMinSpeed = Time;
+		}
+	}
+
+	SampleInterval = 1.f / SampleRate;
+	NumSteps = AnimLength / SampleInterval;
+	for (int32 Step = 0; Step < NumSteps; ++Step)
+	{
+		const float Time = Step * SampleInterval;
+
+		// Assume that during any time before the stop/pivot point, the animation is approaching that point.
+		// TODO: This works for clips that are broken into starts/stops/pivots, but needs to be rethought for more complex clips.
+		const float ValueSign = (Time < TimeOfMinSpeed) ? -1.0f : 1.0f;
+
+		const FVector RootMotionTranslation = Animation->ExtractRootMotionFromRange(TimeOfMinSpeed, Time).GetTranslation();
+		UAnimationBlueprintLibrary::AddFloatCurveKey(Animation, CurveName, Time, ValueSign * RootMotionTranslation.Size2D());
+	}
+}
+
+void UDistanceCurveModifier::OnRevert_Implementation(UAnimSequence* Animation)
+{
+	const bool bRemoveNameFromSkeleton = false;
+	UAnimationBlueprintLibrary::RemoveCurve(Animation, CurveName, bRemoveNameFromSkeleton);
+}
\ No newline at end of file
diff --git a/Engine/Plugins/Animation/AnimationLocomotionLibrary/Source/Editor/Public/AnimationLocomotionLibraryEditor.h b/Engine/Plugins/Animation/AnimationLocomotionLibrary/Source/Editor/Public/AnimationLocomotionLibraryEditor.h
new file mode 100644
--- /dev/null
+++ b/Engine/Plugins/Animation/AnimationLocomotionLibrary/Source/Editor/Public/AnimationLocomotionLibraryEditor.h
@@ -0,0 +1,10 @@
+// Copyright Epic Games, Inc. All Rights Reserved.
+
+#pragma once
+
+#include "Modules/ModuleInterface.h"
+#include "Modules/ModuleManager.h"
+
+class IAnimationLocomotionLibraryEditorModule : public IModuleInterface
+{
+};
\ No newline at end of file
diff --git a/Engine/Plugins/Animation/AnimationLocomotionLibrary/Source/Editor/Public/DistanceCurveModifier.h b/Engine/Plugins/Animation/AnimationLocomotionLibrary/Source/Editor/Public/DistanceCurveModifier.h
new file mode 100644
--- /dev/null
+++ b/Engine/Plugins/Animation/AnimationLocomotionLibrary/Source/Editor/Public/DistanceCurveModifier.h
@@ -0,0 +1,34 @@
+// Copyright Epic Games, Inc. All Rights Reserved.
+
+#pragma once
+
+#include "CoreMinimal.h"
+#include "AnimationModifier.h"
+#include "DistanceCurveModifier.generated.h"
+
+/** Extracts traveling distance information from root motion and bakes it to a curve.
+ * A negative value indicates distance remaining to a stop or pivot point.
+ * A positive value indicates distance traveled from a start point or from the beginning of the clip.
+ */
+UCLASS()
+class UDistanceCurveModifier : public UAnimationModifier
+{
+	GENERATED_BODY()
+
+public:
+
+	/** Rate used to sample the animation. */
+	UPROPERTY(EditAnywhere, Category = Settings, meta = (ClampMin = "1"))
+	int32 SampleRate = 30;
+
+	/** Name for the generated curve. */
+	UPROPERTY(EditAnywhere, BlueprintReadWrite, Category = Settings)
+	FName CurveName = "Distance";
+
+	/** Root motion speed must be below this threshold to be considered stopped. */
+	UPROPERTY(EditAnywhere, BlueprintReadWrite, Category = Settings)
+	float StopSpeedThreshold = 5.0f;
+
+	virtual void OnApply_Implementation(UAnimSequence* Animation) override;
+	virtual void OnRevert_Implementation(UAnimSequence* Animation) override;
+};
\ No newline at end of file
diff --git a/Engine/Plugins/Animation/AnimationLocomotionLibrary/Source/Runtime/AnimationLocomotionLibraryRuntime.Build.cs b/Engine/Plugins/Animation/AnimationLocomotionLibrary/Source/Runtime/AnimationLocomotionLibraryRuntime.Build.cs
new file mode 100644
--- /dev/null
+++ b/Engine/Plugins/Animation/AnimationLocomotionLibrary/Source/Runtime/AnimationLocomotionLibraryRuntime.Build.cs
@@ -0,0 +1,25 @@
+// Copyright Epic Games, Inc. All Rights Reserved.
+
+namespace UnrealBuildTool.Rules
+{
+	public class AnimationLocomotionLibraryRuntime : ModuleRules
+	{
+		public AnimationLocomotionLibraryRuntime(ReadOnlyTargetRules Target) : base(Target)
+		{
+			PublicDependencyModuleNames.AddRange(
+			new string[]
+			{
+                "AnimGraphRuntime",
+                "AnimationCore",
+                "Core",
+                "CoreUObject",
+                "Engine",
+            });
+
+			PrivateDependencyModuleNames.AddRange(
+			new string[]
+			{
+            });
+		}
+	}
+}
diff --git a/Engine/Plugins/Animation/AnimationLocomotionLibrary/Source/Runtime/Private/AnimCharacterMovementLibrary.cpp b/Engine/Plugins/Animation/AnimationLocomotionLibrary/Source/Runtime/Private/AnimCharacterMovementLibrary.cpp
new file mode 100644
--- /dev/null
+++ b/Engine/Plugins/Animation/AnimationLocomotionLibrary/Source/Runtime/Private/AnimCharacterMovementLibrary.cpp
@@ -0,0 +1,188 @@
+// Copyright Epic Games, Inc. All Rights Reserved.
+
+#include "AnimCharacterMovementLibrary.h"
+#include "Animation/AnimSequenceBase.h"
+#include "KismetAnimationLibrary.h"
+
+DEFINE_LOG_CATEGORY_STATIC(LogAnimCharacterMovementLibrary, Verbose, All);
+
+void UAnimCharacterMovementLibrary::UpdateCharacterMovementSnapshot(const FTransform& WorldTransform, const FVector& WorldVelocity, const FVector& WorldAcceleration, bool bIsOnGround,
+	float RootYawOffset, UPARAM(ref) FAnimCharacterMovementSnapshot& Snapshot)
+{
+	// Position
+
+	const FVector WorldLocation = WorldTransform.GetLocation();
+	Snapshot.Distance2DTraveledSinceLastUpdate = FVector::Dist2D(Snapshot.WorldLocation, WorldLocation);
+	Snapshot.WorldLocation = WorldLocation;
+
+	// Velocity
+
+	Snapshot.WorldVelocity = WorldVelocity;
+	Snapshot.LocalVelocity = WorldTransform.InverseTransformVectorNoScale(Snapshot.WorldVelocity);
+	Snapshot.Speed2D = Snapshot.WorldVelocity.Size2D();
+
+	// Acceleration
+
+	Snapshot.WorldAcceleration = WorldAcceleration;
+	Snapshot.LocalAcceleration = WorldTransform.InverseTransformVectorNoScale(Snapshot.WorldAcceleration);
+	Snapshot.AccelerationSize2D = Snapshot.WorldAcceleration.Size2D();
+
+	// Movement angle
+
+	const FRotator Rotation = WorldTransform.GetRotation().Rotator();
+
+	if (FMath::IsNearlyZero(Snapshot.Speed2D))
+	{
+		Snapshot.VelocityYawAngle = 0.0f;
+		Snapshot.AccelerationYawAngle = 0.0f;
+	}
+	else
+	{
+		Snapshot.VelocityYawAngle = UKismetAnimationLibrary::CalculateDirection(Snapshot.WorldVelocity, Rotation);
+		Snapshot.VelocityYawAngle = FRotator::NormalizeAxis(Snapshot.VelocityYawAngle - RootYawOffset);
+
+		Snapshot.AccelerationYawAngle = UKismetAnimationLibrary::CalculateDirection(Snapshot.WorldAcceleration, Rotation);
+		Snapshot.AccelerationYawAngle = FRotator::NormalizeAxis(Snapshot.AccelerationYawAngle - RootYawOffset);
+	}
+
+	// Movement state
+
+	Snapshot.bIsOnGround = bIsOnGround;
+	
+#if !(UE_BUILD_SHIPPING || UE_BUILD_TEST)
+	FString DebugString = FString::Printf(TEXT(
+		"WorldVelocity(%s) | LocalVelocity(%s) | Speed2D(%.3f) | Distance2DTraveledSinceLastUpdate(%.3f)\n"
+		"WorldAcceleration(%s) | LocalAcceleration(%s) | AccelerationSize2D(%.3f)\n"
+		"VelocityYawAngle(%.3f) | AccelerationYawAngle(%.3f)\n"
+		"bIsOnGround(%d)\n"),
+		*Snapshot.WorldVelocity.ToString(), *Snapshot.LocalVelocity.ToString(), Snapshot.Speed2D, Snapshot.Distance2DTraveledSinceLastUpdate,
+		*Snapshot.WorldAcceleration.ToString(), *Snapshot.LocalAcceleration.ToString(), Snapshot.AccelerationSize2D,
+		Snapshot.VelocityYawAngle, Snapshot.AccelerationYawAngle,
+		Snapshot.bIsOnGround);
+
+	UE_LOG(LogAnimCharacterMovementLibrary, VeryVerbose, TEXT("%s"), *DebugString);
+#endif // !(UE_BUILD_SHIPPING || UE_BUILD_TEST)
+}
+
+EAnimCardinalDirection UAnimCharacterMovementLibrary::GetCardinalDirectionFromAngle(EAnimCardinalDirection PreviousCardinalDirection, float AngleInDegrees, float DeadZoneAngle)
+{
+	// Use deadzone offset to favor backpedal on S & W, and favor frontpedal on N & E
+	const float AbsoluteAngle = FMath::Abs(AngleInDegrees);
+	if (PreviousCardinalDirection == EAnimCardinalDirection::North)
+	{
+		if (AbsoluteAngle <= (45.f + DeadZoneAngle + DeadZoneAngle))
+		{
+			return EAnimCardinalDirection::North;
+		}
+		else if (AbsoluteAngle >= 135.f - DeadZoneAngle)
+		{
+			return EAnimCardinalDirection::South;
+		}
+		else if (AngleInDegrees > 0.f)
+		{
+			return EAnimCardinalDirection::East;
+		}
+		return EAnimCardinalDirection::West;
+	}
+	else if (PreviousCardinalDirection == EAnimCardinalDirection::South)
+	{
+		if (AbsoluteAngle <= 45.f + DeadZoneAngle)
+		{
+			return EAnimCardinalDirection::North;
+		}
+		else if (AbsoluteAngle >= (135.f - DeadZoneAngle - DeadZoneAngle))
+		{
+			return EAnimCardinalDirection::South;
+		}
+		else if (AngleInDegrees > 0.f)
+		{
+			return EAnimCardinalDirection::East;
+		}
+		return EAnimCardinalDirection::West;
+	}
+
+	// East and West
+	if (AbsoluteAngle <= (45.f + DeadZoneAngle))
+	{
+		return EAnimCardinalDirection::North;
+	}
+	else if (AbsoluteAngle >= (135.f - DeadZoneAngle))
+	{
+		return EAnimCardinalDirection::South;
+	}
+	else if (AngleInDegrees > 0.f)
+	{
+		return EAnimCardinalDirection::East;
+	}
+	return EAnimCardinalDirection::West;
+}
+
+const UAnimSequence* UAnimCharacterMovementLibrary::SelectAnimForCardinalDirection(EAnimCardinalDirection CardinalDirection, const FCardinalDirectionAnimSet& AnimSet)
+{
+	switch (CardinalDirection)
+	{
+		case EAnimCardinalDirection::North:
+			return AnimSet.NorthAnim;
+		case EAnimCardinalDirection::East:
+			return AnimSet.EastAnim;
+		case EAnimCardinalDirection::South:
+			return AnimSet.SouthAnim;
+		case EAnimCardinalDirection::West:
+			return AnimSet.WestAnim;
+		default:
+			checkNoEntry();
+			return nullptr;
+	}
+
+	return nullptr;
+}
+
+FVector UAnimCharacterMovementLibrary::PredictGroundMovementStopLocation(const FAnimCharacterMovementSnapshot& MovementSnapshot, const FAnimCharacterMovementPredictionSnapshot& PredictionSnapshot)
+{
+	FVector PredictedStopLocation = FVector::ZeroVector;
+
+	float ActualBrakingFriction = (PredictionSnapshot.bUseSeparateBrakingFriction ? PredictionSnapshot.BrakingFriction : PredictionSnapshot.GroundFriction);
+	const float FrictionFactor = FMath::Max(0.f, PredictionSnapshot.BrakingFrictionFactor);
+	ActualBrakingFriction = FMath::Max(0.f, ActualBrakingFriction * FrictionFactor);
+	float BrakingDeceleration = FMath::Max(0.f, PredictionSnapshot.BrakingDecelerationWalking);
+
+	const FVector WorldVelocity2D = MovementSnapshot.WorldVelocity * FVector(1.f, 1.f, 0.f);
+	FVector WorldVelocityDir2D;
+	float Speed2D;
+	WorldVelocity2D.ToDirectionAndLength(WorldVelocityDir2D, Speed2D);
+
+	const float Divisor = ActualBrakingFriction * Speed2D + BrakingDeceleration;
+	if (Divisor > 0.f)
+	{
+		const float TimeToStop = Speed2D / Divisor;
+		PredictedStopLocation = WorldVelocity2D * TimeToStop + 0.5f * ((-ActualBrakingFriction) * WorldVelocity2D - BrakingDeceleration * WorldVelocityDir2D) * TimeToStop * TimeToStop;
+	}
+
+	return PredictedStopLocation;
+}
+
+FVector UAnimCharacterMovementLibrary::PredictGroundMovementPivotLocation(const FAnimCharacterMovementSnapshot& MovementSnapshot, float GroundFriction)
+{
+	FVector PredictedPivotLocation = FVector::ZeroVector;
+
+	const FVector WorldAcceleration2D = MovementSnapshot.WorldAcceleration * FVector(1.f, 1.f, 0.f);
+	
+	FVector WorldAccelerationDir2D;
+	float WorldAccelerationSize2D;
+	WorldAcceleration2D.ToDirectionAndLength(WorldAccelerationDir2D, WorldAccelerationSize2D);
+
+	const float VelocityAlongAcceleration = (MovementSnapshot.WorldVelocity | WorldAccelerationDir2D);
+	if (VelocityAlongAcceleration < 0.0f)
+	{
+		const float SpeedAlongAcceleration = -VelocityAlongAcceleration;
+		const float Divisor = WorldAccelerationSize2D + 2.f * SpeedAlongAcceleration * GroundFriction;
+		const float TimeToDirectionChange = SpeedAlongAcceleration / Divisor;
+
+		const FVector AccelerationForce = MovementSnapshot.WorldAcceleration - 
+			(MovementSnapshot.WorldVelocity - WorldAccelerationDir2D * MovementSnapshot.Speed2D) * GroundFriction;
+	
+		PredictedPivotLocation = MovementSnapshot.WorldVelocity * TimeToDirectionChange + 0.5f * AccelerationForce * TimeToDirectionChange * TimeToDirectionChange;
+	}
+
+	return PredictedPivotLocation;
+}
\ No newline at end of file
diff --git a/Engine/Plugins/Animation/AnimationLocomotionLibrary/Source/Runtime/Private/AnimDistanceMatchingLibrary.cpp b/Engine/Plugins/Animation/AnimationLocomotionLibrary/Source/Runtime/Private/AnimDistanceMatchingLibrary.cpp
new file mode 100644
--- /dev/null
+++ b/Engine/Plugins/Animation/AnimationLocomotionLibrary/Source/Runtime/Private/AnimDistanceMatchingLibrary.cpp
@@ -0,0 +1,155 @@
+// Copyright Epic Games, Inc. All Rights Reserved.
+
+#include "AnimDistanceMatchingLibrary.h"
+#include "Animation/AnimSequence.h"
+#include "AnimNodes/AnimNode_SequenceEvaluator.h"
+
+DEFINE_LOG_CATEGORY_STATIC(LogAnimDistanceMatchingLibrary, Verbose, All);
+
+/**
+ * Advance from the current time to a new time in the animation that will result in the desired distance traveled by the authored root motion.
+ */
+float GetTimeAfterDistanceTraveled(const UAnimSequence* AnimSequence, float CurrentTime, float DistanceTraveled, const FDistanceCurve& CachedDistanceCurve, const bool bAllowLooping)
+{
+	float NewTime = CurrentTime;
+	if (AnimSequence != nullptr)
+	{
+		// Avoid infinite loops if the animation doesn't cover any distance.
+		if (!FMath::IsNearlyZero(CachedDistanceCurve.GetDistanceRange(AnimSequence)))
+		{
+			float AccumulatedDistance = 0.f;
+
+			const float SequenceLength = AnimSequence->GetPlayLength();
+			const float StepTime = 1.f / 30.f;
+
+			// Traverse the distance curve, accumulating animated distance until the desired distance is reached.
+			while ((AccumulatedDistance < DistanceTraveled) && (bAllowLooping || (NewTime + StepTime < SequenceLength)))
+			{
+				const float CurrentDistance = CachedDistanceCurve.GetValueAtPosition(AnimSequence, NewTime);
+				const float DistanceAfterStep = CachedDistanceCurve.GetValueAtPosition(AnimSequence, NewTime + StepTime);
+				const float AnimationDistanceThisStep = DistanceAfterStep - CurrentDistance;
+
+				if (!FMath::IsNearlyZero(AnimationDistanceThisStep))
+				{
+					// Keep advancing if the desired distance hasn't been reached.
+					if (AccumulatedDistance + AnimationDistanceThisStep < DistanceTraveled)
+					{
+						FAnimationRuntime::AdvanceTime(bAllowLooping, StepTime, NewTime, SequenceLength);
+						AccumulatedDistance += AnimationDistanceThisStep;
+					}
+					// Once the desired distance is passed, find the approximate time between samples where the distance will be reached.
+					else
+					{
+						const float DistanceAlpha = (DistanceTraveled - AccumulatedDistance) / AnimationDistanceThisStep;
+						FAnimationRuntime::AdvanceTime(bAllowLooping, DistanceAlpha * StepTime, NewTime, SequenceLength);
+						AccumulatedDistance = DistanceTraveled;
+						break;
+					}
+				}
+			}
+		}
+		else
+		{
+			UE_LOG(LogAnimDistanceMatchingLibrary, Warning, TEXT("Anim sequence (%s) is missing a distance curve or doesn't cover enough distance for GetTimeAfterDistanceTraveled."), *GetNameSafe(AnimSequence));
+		}
+	}
+	else
+	{
+		UE_LOG(LogAnimDistanceMatchingLibrary, Warning, TEXT("Invalid AnimSequence passed to GetTimeAfterDistanceTraveled"));
+	}
+
+	return NewTime;
+}
+
+FSequenceEvaluatorReference UAnimDistanceMatchingLibrary::AdvanceTimeByDistanceMatching(const FAnimUpdateContext& UpdateContext, const FSequenceEvaluatorReference& SequenceEvaluator,
+	float DistanceTraveled, const FDistanceCurve& CachedDistanceCurve, const UCurveFloat* PlayRateAdjustmentCurve, FVector2D PlayRateClamp)
+{
+	SequenceEvaluator.CallAnimNodeFunction<FAnimNode_SequenceEvaluator>(
+		TEXT("AdvanceTimeByDistanceMatching"),
+		[&UpdateContext, DistanceTraveled, &CachedDistanceCurve, PlayRateAdjustmentCurve, PlayRateClamp](FAnimNode_SequenceEvaluator& InSequenceEvaluator)
+		{
+			if (const FAnimationUpdateContext* AnimationUpdateContext = UpdateContext.GetContext())
+			{
+				const float DeltaTime = AnimationUpdateContext->GetDeltaTime(); 
+
+				if (const UAnimSequence* AnimSequence = Cast<UAnimSequence>(InSequenceEvaluator.GetSequence()))
+				{
+					const float CurrentTime = InSequenceEvaluator.GetExplicitTime();
+					const float CurrentAssetLength = InSequenceEvaluator.GetCurrentAssetLength();
+					const bool bAllowLooping = InSequenceEvaluator.GetShouldLoop();
+
+					float TimeAfterDistanceTraveled = GetTimeAfterDistanceTraveled(AnimSequence, CurrentTime, DistanceTraveled, CachedDistanceCurve, bAllowLooping);
+
+					// Calculate the effective playrate that would result from advancing the animation by the distance traveled.
+					// Account for the animation looping.
+					if (TimeAfterDistanceTraveled < CurrentTime)
+					{
+						TimeAfterDistanceTraveled += CurrentAssetLength;
+					}
+					float EffectivePlayRate = (TimeAfterDistanceTraveled - CurrentTime) / DeltaTime;
+
+					// Clamp the effective play rate.
+					if (PlayRateAdjustmentCurve != nullptr)
+					{
+						EffectivePlayRate = FMath::Max(PlayRateAdjustmentCurve->GetFloatValue(EffectivePlayRate), 0.f);
+					}
+					else if (PlayRateClamp.X < PlayRateClamp.Y)
+					{
+						EffectivePlayRate = FMath::Clamp(EffectivePlayRate, PlayRateClamp.X, PlayRateClamp.Y);
+					}
+
+					// Advance animation time by the effective play rate.
+					float NewTime = CurrentTime;
+					FAnimationRuntime::AdvanceTime(bAllowLooping, EffectivePlayRate * DeltaTime, NewTime, CurrentAssetLength);
+
+					if(!InSequenceEvaluator.SetExplicitTime(NewTime))
+					{
+						UE_LOG(LogAnimDistanceMatchingLibrary, Warning, TEXT("Could not set explicit time on sequence evaluator, value is not dynamic. Set it as Always Dynamic."));
+					}
+				}
+				else
+				{
+					UE_LOG(LogAnimDistanceMatchingLibrary, Warning, TEXT("Sequence evaluator does not have an anim sequence to play."));
+				}
+			}
+			else
+			{
+				UE_LOG(LogAnimDistanceMatchingLibrary, Warning, TEXT("AdvanceTimeByDistanceMatching called with invalid context"));
+			}
+		});
+
+	return SequenceEvaluator;
+}
+
+FSequenceEvaluatorReference UAnimDistanceMatchingLibrary::DistanceMatchToTarget(const FSequenceEvaluatorReference& SequenceEvaluator,
+	float DistanceToTarget, const FDistanceCurve& CachedDistanceCurve)
+{
+	SequenceEvaluator.CallAnimNodeFunction<FAnimNode_SequenceEvaluator>(
+		TEXT("DistanceMatchToTarget"),
+		[DistanceToTarget, &CachedDistanceCurve](FAnimNode_SequenceEvaluator& InSequenceEvaluator)
+		{
+			if (const UAnimSequence* AnimSequence = Cast<UAnimSequence>(InSequenceEvaluator.GetSequence()))
+			{
+				if (CachedDistanceCurve.IsValid(AnimSequence))
+				{
+					// By convention, distance curves store the distance to a target as a negative value.
+					const float NewTime = CachedDistanceCurve.GetAnimPositionFromDistance(AnimSequence, -DistanceToTarget);
+					if (!InSequenceEvaluator.SetExplicitTime(NewTime))
+					{
+						UE_LOG(LogAnimDistanceMatchingLibrary, Warning, TEXT("Could not set explicit time on sequence evaluator, value is not dynamic. Set it as Always Dynamic."));
+					}
+				}
+				else
+				{
+					UE_LOG(LogAnimDistanceMatchingLibrary, Warning, TEXT("DistanceMatchToTarget called with invalid CachedDistanceCurve."));
+				}
+			}
+			else
+			{
+				UE_LOG(LogAnimDistanceMatchingLibrary, Warning, TEXT("Sequence evaluator does not have an anim sequence to play."));
+			}
+			
+		});
+
+	return SequenceEvaluator;
+}
\ No newline at end of file
diff --git a/Engine/Plugins/Animation/AnimationLocomotionLibrary/Source/Runtime/Private/AnimTurnInPlaceLibrary.cpp b/Engine/Plugins/Animation/AnimationLocomotionLibrary/Source/Runtime/Private/AnimTurnInPlaceLibrary.cpp
new file mode 100644
--- /dev/null
+++ b/Engine/Plugins/Animation/AnimationLocomotionLibrary/Source/Runtime/Private/AnimTurnInPlaceLibrary.cpp
@@ -0,0 +1,9 @@
+// Copyright Epic Games, Inc. All Rights Reserved.
+
+#include "AnimTurnInPlaceLibrary.h"
+
+void UAnimTurnInPlaceLibrary::UpdateTurnInPlace(float DeltaTime, bool bAllowTurnInPlace, bool bHoldYawOffset, bool bIsTurnTransitionStateRelevant,
+	const FRotator& MeshWorldRotation, const FAnimTurnInPlaceAnimSet& AnimSet, UPARAM(ref) FAnimTurnInPlaceState& TurnInPlaceState)
+{
+	TurnInPlaceState.Update(DeltaTime, bAllowTurnInPlace, bHoldYawOffset, bIsTurnTransitionStateRelevant, MeshWorldRotation, AnimSet, TurnInPlaceState);
+}
\ No newline at end of file
diff --git a/Engine/Plugins/Animation/AnimationLocomotionLibrary/Source/Runtime/Private/AnimTurnInPlaceTypes.cpp b/Engine/Plugins/Animation/AnimationLocomotionLibrary/Source/Runtime/Private/AnimTurnInPlaceTypes.cpp
new file mode 100644
--- /dev/null
+++ b/Engine/Plugins/Animation/AnimationLocomotionLibrary/Source/Runtime/Private/AnimTurnInPlaceTypes.cpp
@@ -0,0 +1,152 @@
+// Copyright Epic Games, Inc. All Rights Reserved.
+
+#include "AnimTurnInPlaceTypes.h"
+
+DEFINE_LOG_CATEGORY_STATIC(LogAnimTurnInPlaceLibrary, Verbose, All);
+
+int32 FAnimTurnInPlaceAnimSet::FindBestTurnTransitionIndex(float YawOffset) const
+{
+	int32 BestTurnTransitionIndex = INDEX_NONE;
+	float BestResultingOffsetSize = MAX_FLT;
+
+	for (int32 Index = 0; Index < TurnTransitions.Num(); Index++)
+	{
+		if (const UAnimSequence* TurnTransitionAnim = TurnTransitions[Index].Anim)
+		{
+			// TODO: We could potentially cache the rotation info in the struct or a curve in the asset.
+			const FTransform TransitionAnimRootMotion = TurnTransitionAnim->ExtractRootMotionFromRange(0.f, TurnTransitionAnim->GetPlayLength());
+			const float TransitionAnimYaw = FMath::RadiansToDegrees(TransitionAnimRootMotion.GetRotation().GetTwistAngle(FVector::UpVector));
+
+			// Calculate the yaw offset that would result from playing the animation.
+			const float OffsetAfterTransition = FRotator::NormalizeAxis(YawOffset + TransitionAnimYaw);
+			const float OffsetSizeAfterTransition = FMath::Abs(OffsetAfterTransition);
+
+			// Check if the animation will result in a smaller yaw offset.
+			if ((OffsetSizeAfterTransition + TurnDeadZoneAngle) < FMath::Abs(YawOffset))
+			{
+				const bool bTransitionResultsInSmallerAngle = (OffsetSizeAfterTransition < BestResultingOffsetSize);
+
+				// Prefer animations that rotate toward the offset.
+				// E.g. If the offset is -160 degrees, we prefer to rotate 180 degrees toward the offset, rather than 180 away, even though they will result in the same new offset.
+				const bool bSameResultButBetterDirection = FMath::IsNearlyEqual(OffsetSizeAfterTransition, BestResultingOffsetSize) && (TransitionAnimYaw * YawOffset < 0.f);
+
+				if (bTransitionResultsInSmallerAngle || bSameResultButBetterDirection)
+				{
+					BestTurnTransitionIndex = Index;
+					BestResultingOffsetSize = OffsetSizeAfterTransition;
+				}
+			}
+		}
+	}
+
+	return BestTurnTransitionIndex;
+}
+
+void FAnimTurnInPlaceState::Update(float DeltaTime, bool bAllowTurnInPlace, bool bHoldYawOffset, bool bIsTurnTransitionStateRelevant,
+	const FRotator& MeshWorldRotation, const FAnimTurnInPlaceAnimSet& AnimSet, UPARAM(ref) FAnimTurnInPlaceState& TurnInPlaceState)
+{
+	bTurnTransitionRequested = false;
+	bTurnRecoveryRequested = false;
+
+	// Avoid large offset from default value on first update.
+	if (MeshWorldYaw == FLT_MAX)
+	{
+		MeshWorldYaw = MeshWorldRotation.Yaw;
+	}
+
+	// Calculate how much the mesh has rotated since the last update.
+	const float YawDeltaSinceLastUpdate = MeshWorldRotation.Yaw - MeshWorldYaw;
+	MeshWorldYaw = MeshWorldRotation.Yaw;
+
+	if (bAllowTurnInPlace)
+	{
+		// Apply any recent mesh rotation to the root offset.
+		RootYawOffset = FRotator::NormalizeAxis(RootYawOffset - YawDeltaSinceLastUpdate);
+
+		// Update the current turn transition.
+		if (bIsTurnTransitionStateRelevant)
+		{
+			UpdateActiveTurnTransition(DeltaTime);
+		}
+		// Trigger a turn transition animation if necessary.
+		else
+		{
+			UpdateTurnTransitionTrigger(DeltaTime, AnimSet);
+		}
+	}
+	// Interpolate out the offset if it's no longer requested.
+	else if (!bHoldYawOffset)
+	{
+		RootYawOffset = YawOffsetInterpolator.Update(RootYawOffset, 0.f, DeltaTime);
+	}
+
+	RootYawOffsetInverse = -RootYawOffset;
+
+#if !(UE_BUILD_SHIPPING || UE_BUILD_TEST)
+	FString LogMessage = FString::Printf(TEXT("RootYawOffset(%.2f) PendingTurnAnim(%s) PendingTurnDelayCounter(%.2f) ActiveTurnAnim(%s) ActiveTurnAnimTime(%.2f) bTurnTransitionRequested(%d) bHoldYawOffset(%d) bIsTurnTransitionStateRelevant(%d)"), 
+		RootYawOffset, *GetNameSafe(PendingTurnAnim), PendingTurnDelayCounter, *GetNameSafe(ActiveTurnAnim), ActiveTurnAnimTime, bTurnTransitionRequested, bHoldYawOffset, bIsTurnTransitionStateRelevant);
+	UE_LOG(LogAnimTurnInPlaceLibrary, VeryVerbose, TEXT("%s"), *LogMessage);
+#endif // !(UE_BUILD_SHIPPING || UE_BUILD_TEST)
+}
+
+void FAnimTurnInPlaceState::UpdateActiveTurnTransition(float DeltaTime)
+{
+	check(ActiveTurnAnim);
+
+	// Manually advance the turn transition animation so we can subtract the animated change in rotation from the current yaw offset.
+	// TODO: This could be a lot simpler if we could leverage the root motion animation attribute somehow to consume the yaw offset.
+	const float EndTime = ActiveTurnAnim->GetPlayLength();
+	const float NewTime = FMath::Min(ActiveTurnAnimTime + DeltaTime, EndTime);
+
+	if (NewTime > ActiveTurnAnimTime)
+	{
+		const FTransform AnimRootMotion = ActiveTurnAnim->ExtractRootMotionFromRange(ActiveTurnAnimTime, NewTime);
+		const float AnimYaw = FMath::RadiansToDegrees(AnimRootMotion.GetRotation().GetTwistAngle(FVector::UpVector));
+
+		ActiveTurnAnimTime = NewTime;
+		RootYawOffset += AnimYaw;
+
+		// Check if there's any rotation left in the animation.
+		const FTransform AnimRemainingRootMotion = ActiveTurnAnim->ExtractRootMotionFromRange(NewTime, EndTime);
+		const float AnimRemainingYaw = FMath::RadiansToDegrees(AnimRemainingRootMotion.GetRotation().GetTwistAngle(FVector::UpVector));
+		if (FMath::IsNearlyZero(AnimRemainingYaw))
+		{
+			TurnRecoveryAnim = ActiveTurnAnim;
+			TurnRecoveryAnimStartTime = ActiveTurnAnimTime;
+			bTurnRecoveryRequested = true;
+		}
+	}
+	// Turn transition animations should have some recovery time at the end. A possible workaround if they don't is to have an automatic transition
+	// so the state machine doesn't get stuck in the turn transition state.
+	else
+	{
+		ensureMsgf(false, TEXT("Reached end of turn transition without reaching the end of rotation. If the animation (%s) rotates until the end, an automatic transition back to idle is recommended."), *GetNameSafe(ActiveTurnAnim));
+	}
+}
+
+void FAnimTurnInPlaceState::UpdateTurnTransitionTrigger(float DeltaTime, const FAnimTurnInPlaceAnimSet& AnimSet)
+{
+	const UAnimSequence* PrevPendingTurnAnim = PendingTurnAnim;
+
+	const int32 BestTurnTransitionIndex = AnimSet.FindBestTurnTransitionIndex(RootYawOffset);
+	if (BestTurnTransitionIndex != INDEX_NONE)
+	{
+		// If a turn transition is desired, handle the required delay before playing it. If a new turn transition is requested during
+		// the delay (e.g. because the character keeps rotating), restart the delay.
+		const FAnimTurnTransition& PendingTurnTransition = AnimSet.TurnTransitions[BestTurnTransitionIndex];
+		PendingTurnAnim = PendingTurnTransition.Anim;
+		PendingTurnDelayCounter = (PendingTurnAnim == PrevPendingTurnAnim) ? (PendingTurnDelayCounter + DeltaTime) : 0.f;
+
+		if (PendingTurnDelayCounter >= PendingTurnTransition.DelayBeforeTrigger)
+		{
+			ActiveTurnAnim = PendingTurnAnim;
+			ActiveTurnAnimTime = 0.f;
+			bTurnTransitionRequested = true;
+			PendingTurnAnim = nullptr;
+		}
+	}
+	else
+	{
+		PendingTurnAnim = nullptr;
+	}
+}
\ No newline at end of file
diff --git a/Engine/Plugins/Animation/AnimationLocomotionLibrary/Source/Runtime/Private/AnimationLocomotionLibraryModule.cpp b/Engine/Plugins/Animation/AnimationLocomotionLibrary/Source/Runtime/Private/AnimationLocomotionLibraryModule.cpp
new file mode 100644
--- /dev/null
+++ b/Engine/Plugins/Animation/AnimationLocomotionLibrary/Source/Runtime/Private/AnimationLocomotionLibraryModule.cpp
@@ -0,0 +1,19 @@
+// Copyright Epic Games, Inc. All Rights Reserved.
+
+#include "AnimationLocomotionLibraryModule.h"
+
+#define LOCTEXT_NAMESPACE "AnimationLocomotionModule"
+
+void FAnimationLocomotionLibraryModule::StartupModule()
+{
+
+}
+
+void FAnimationLocomotionLibraryModule::ShutdownModule()
+{
+
+}
+
+#undef LOCTEXT_NAMESPACE
+
+IMPLEMENT_MODULE(FAnimationLocomotionLibraryModule, AnimationLocomotionLibrary)
\ No newline at end of file
diff --git a/Engine/Plugins/Animation/AnimationLocomotionLibrary/Source/Runtime/Public/AnimCharacterMovementLibrary.h b/Engine/Plugins/Animation/AnimationLocomotionLibrary/Source/Runtime/Public/AnimCharacterMovementLibrary.h
new file mode 100644
--- /dev/null
+++ b/Engine/Plugins/Animation/AnimationLocomotionLibrary/Source/Runtime/Public/AnimCharacterMovementLibrary.h
@@ -0,0 +1,75 @@
+// Copyright Epic Games, Inc. All Rights Reserved.
+
+#pragma once
+
+#include "AnimCharacterMovementTypes.h"
+#include "AnimCharacterMovementLibrary.generated.h"
+
+class UAnimSequence;
+
+/**
+  * Library of common techniques for driving character locomotion animations. 
+  */
+UCLASS()
+class ANIMATIONLOCOMOTIONLIBRARYRUNTIME_API UAnimCharacterMovementLibrary : public UBlueprintFunctionLibrary
+{
+	GENERATED_BODY()
+
+public:
+
+	/**
+	 * Populate a snapshot struct with movement data that's commonly used by animation graph logic.
+	 * To avoid performance costs from calling this in the Event Graph on the game thread, it's recommended to call it in BlueprintThreadSafeUpdateAnimation and
+ 	 * use the Property Access system to access the input parameters (Property Access will handle copying the inputs at the right time in the frame).
+	 * @param WorldTransform - The transform of the character in world space.
+	 * @param WorldVelocity - The velocity of the character in world space.
+	 * @param WorldAcceleration - The acceleration of the character in world space.
+	 * @param bIsOnGround - If the character is on the ground.
+	 * @param RootYawOffset - Offset being applied to the root bone in the animation graph (e.g. for countering capsule rotation). Set to zero if not needed.
+	 * @param Snapshot - The snapshot to write to. This is typically a member variable of the animation blueprint.
+	 */
+	UFUNCTION(BlueprintCallable, Category = "Animation Character Movement", meta = (BlueprintThreadSafe))
+	static void UpdateCharacterMovementSnapshot(const FTransform& WorldTransform, const FVector& WorldVelocity, const FVector& WorldAcceleration, bool bIsOnGround,
+		float RootYawOffset, UPARAM(ref) FAnimCharacterMovementSnapshot& Snapshot);
+
+	/**
+	 * Calculate the closest cardinal direction to the direction the character is currently moving.
+	 * @param PreviousCardinalDirection - The cardinal direction from the previous frame. Typically the animation blueprint holds a EAnimCardinalDirection variable.
+	 * @param DirectionAngleInDegrees - The direction that the character is currently moving. FAnimCharacterMovementSnapshot.VelocityYawAngle is a commonly used input for this.
+	 * @param DirectionDeadZoneAngle - Deadzone to prevent flickering between directions at angle boundaries.
+	 * @return The resulting cardinal direction.
+	 */
+	UFUNCTION(BlueprintPure, Category = "Animation Character Movement", meta = (BlueprintThreadSafe))
+	static EAnimCardinalDirection GetCardinalDirectionFromAngle(EAnimCardinalDirection PreviousCardinalDirection, float DirectionAngleInDegrees, float DeadZoneAngle);
+
+	/**
+	 * Select an animation to play based on the cardinal direction calculated by GetCardinalDirectionFromAngle(). For example, this can pick a start animation
+	 * to play based on the character's movement direction.	    
+	 * @param CardinalDirection - The closest cardinal direction to the character's movement direction.
+	 * @param AnimSet - The set of animations to choose from.
+	 * @return The animation to play.
+	 */
+	UFUNCTION(BlueprintPure, Category = "Animation Character Movement", meta = (BlueprintThreadSafe))
+	static const UAnimSequence* SelectAnimForCardinalDirection(EAnimCardinalDirection CardinalDirection, const FCardinalDirectionAnimSet& AnimSet);
+
+	/**
+	 * Predict where the character will stop based on its current movement properties and parameters from the movement component.
+	 * This uses prediction logic that is heavily tied to the UCharacterMovementComponent.
+	 * @param MovementSnapshot - Snapshot of current movement properties.
+	 * @param PredictionSnapshot - Snapshot of parameters needed to predict how the movement component will move. Because this is a thread safe function, it's
+	 *	recommended to populate these fields via the Property Access system.
+	 * @return The predicted stop position in local space to the character. The size of this vector will be the distance to stop.
+	 */
+	UFUNCTION(BlueprintPure, Category = "Animation Character Movement", meta = (BlueprintThreadSafe))
+	static FVector PredictGroundMovementStopLocation(const FAnimCharacterMovementSnapshot& MovementSnapshot, const FAnimCharacterMovementPredictionSnapshot& PredictionSnapshot);
+
+	/**
+	 * Predict where the character will change direction during a pivot based on its current movement properties and parameters from the movement component.
+	 * This uses prediction logic that is heavily tied to the UCharacterMovementComponent.
+	 * @param MovementSnapshot - Snapshot of current movement properties.
+	 * @param GroundFriction - Value from the movement component. Because this is a thread safe function, it's recommended to populate this field via the Property Access system.
+	 * @return The predicted pivot position in local space to the character. The size of this vector will be the distance to the pivot.
+	 */
+	UFUNCTION(BlueprintPure, Category = "Animation Character Movement", meta = (BlueprintThreadSafe))
+	static FVector PredictGroundMovementPivotLocation(const FAnimCharacterMovementSnapshot& MovementSnapshot, float GroundFriction);
+};
\ No newline at end of file
diff --git a/Engine/Plugins/Animation/AnimationLocomotionLibrary/Source/Runtime/Public/AnimCharacterMovementTypes.h b/Engine/Plugins/Animation/AnimationLocomotionLibrary/Source/Runtime/Public/AnimCharacterMovementTypes.h
new file mode 100644
--- /dev/null
+++ b/Engine/Plugins/Animation/AnimationLocomotionLibrary/Source/Runtime/Public/AnimCharacterMovementTypes.h
@@ -0,0 +1,115 @@
+// Copyright Epic Games, Inc. All Rights Reserved.
+
+#pragma once
+
+#include "AnimCharacterMovementTypes.generated.h"
+
+class UAnimSequence;
+
+UENUM(BlueprintType)
+enum class EAnimCardinalDirection : uint8
+{
+	North,
+	East,
+	South,
+	West
+};
+
+/**
+ * Animations for a locomotion set authored with only four cardinal directions.
+ * This will often be accompanied by Orientation Warping to account for diagonals.
+ */
+USTRUCT(BlueprintType)
+struct FCardinalDirectionAnimSet
+{
+	GENERATED_BODY()
+
+public:
+
+	UPROPERTY(EditAnywhere, BlueprintReadWrite, Category = Animation)
+	TObjectPtr<UAnimSequence> NorthAnim = nullptr;
+
+	UPROPERTY(EditAnywhere, BlueprintReadWrite, Category = Animation)
+	TObjectPtr<UAnimSequence> EastAnim = nullptr;
+
+	UPROPERTY(EditAnywhere, BlueprintReadWrite, Category = Animation)
+	TObjectPtr<UAnimSequence> SouthAnim = nullptr;
+
+	UPROPERTY(EditAnywhere, BlueprintReadWrite, Category = Animation)
+	TObjectPtr<UAnimSequence> WestAnim = nullptr;
+};
+
+/**
+ * Snapshot of movement properties used to predict where the character will move in the future.
+ * These properties are found on the UCharacterMovementComponent. They're copied (usually via Property Access) on the game thread
+ * so they can be used in thread safe functions during animation update.
+ */
+USTRUCT(BlueprintType)
+struct FAnimCharacterMovementPredictionSnapshot
+{
+	GENERATED_BODY()
+
+public:
+
+	UPROPERTY(BlueprintReadWrite, Category = MovementPrediction)
+	float GroundFriction = 0.0f;
+
+	UPROPERTY(BlueprintReadWrite, Category = MovementPrediction)
+	float BrakingFriction = 0.0f;
+
+	UPROPERTY(BlueprintReadWrite, Category = MovementPrediction)
+	float BrakingFrictionFactor = 0.0f;
+
+	UPROPERTY(BlueprintReadWrite, Category = MovementPrediction)
+	float BrakingDecelerationWalking = 0.0f;
+
+	UPROPERTY(BlueprintReadWrite, Category = MovementPrediction)
+	bool bUseSeparateBrakingFriction = false;
+};
+
+/**
+ * Snapshot of movement data commonly used to drive locomotion animations.
+ * See UAnimationCharacterMovementLibrary::UpdateCharacterMovementSnapshot() for an example of how to populate this data.
+ */
+USTRUCT(BlueprintType)
+struct FAnimCharacterMovementSnapshot
+{
+	GENERATED_BODY()
+
+public:
+
+	UPROPERTY(BlueprintReadWrite, Category = Movement)
+	FVector WorldLocation = FVector::ZeroVector;
+
+	UPROPERTY(BlueprintReadWrite, Category = Movement)
+	FVector WorldVelocity = FVector::ZeroVector;
+
+	UPROPERTY(BlueprintReadWrite, Category = Movement)
+	FVector LocalVelocity = FVector::ZeroVector;
+
+	UPROPERTY(BlueprintReadWrite, Category = Movement)
+	FVector WorldAcceleration = FVector::ZeroVector;
+
+	UPROPERTY(BlueprintReadWrite, Category = Movement)
+	FVector LocalAcceleration = FVector::ZeroVector;
+
+	/** Angle (in degrees) between velocity and the character's forward vector. */
+	UPROPERTY(BlueprintReadWrite, Category = Movement)
+	float VelocityYawAngle = 0.0f;
+
+	/** Angle (in degrees) between acceleration and the character's forward vector. */
+	UPROPERTY(BlueprintReadWrite, Category = Movement)
+	float AccelerationYawAngle = 0.0f;
+
+	UPROPERTY(BlueprintReadWrite, Category = Movement)
+	float Distance2DTraveledSinceLastUpdate = 0.0f;
+
+	UPROPERTY(BlueprintReadWrite, Category = Movement)
+	float Speed2D = 0.0f;
+
+	UPROPERTY(BlueprintReadWrite, Category = Movement)
+	float AccelerationSize2D = 0.0f;
+
+	UPROPERTY(BlueprintReadWrite, Category = Movement)
+	bool bIsOnGround = false;
+};
\ No newline at end of file
diff --git a/Engine/Plugins/Animation/AnimationLocomotionLibrary/Source/Runtime/Public/AnimDistanceMatchingLibrary.h b/Engine/Plugins/Animation/AnimationLocomotionLibrary/Source/Runtime/Public/AnimDistanceMatchingLibrary.h
new file mode 100644
--- /dev/null
+++ b/Engine/Plugins/Animation/AnimationLocomotionLibrary/Source/Runtime/Public/AnimDistanceMatchingLibrary.h
@@ -0,0 +1,49 @@
+// Copyright Epic Games, Inc. All Rights Reserved.
+
+#pragma once
+
+#include "SequenceEvaluatorLibrary.h"
+#include "Animation/AnimCurveTypes.h"
+#include "AnimDistanceMatchingLibrary.generated.h"
+
+class UAnimSequence;
+
+/**
+ * Library of techniques for driving animations by distance metrics rather than by time.
+ * These techniques can be effective at compensating for differences between character movement and authored motion in the animations.
+ * Distance Matching effectively changes the playrate of the animation to keep the feet from sliding. It's common to clamp the resulting
+ * play rate to avoid animations playing too slow or too fast and to use techniques such as Stride Warping to make up the difference.
+ */
+UCLASS()
+class ANIMATIONLOCOMOTIONLIBRARYRUNTIME_API UAnimDistanceMatchingLibrary : public UBlueprintFunctionLibrary
+{
+	GENERATED_BODY()
+
+public:
+	/**
+	 * Advance the sequence evaluator forward by distance traveled rather than time. A distance curve is required on the animation that
+	 * describes the distance traveled by the root bone in the animation. See UDistanceCurveModifier.
+	 * @param UpdateContext - The update context provided in the anim node function.
+	 * @param SequenceEvaluator - The sequence evaluator node to operate on.
+	 * @param DistanceTraveled - The distance traveled by the character since the last animation update. See FAnimCharacterMovementSnapshot.Distance2DTraveledSinceLastUpdate.
+	 * @param CachedDistanceCurve - Optimized access to curve data. This will typically be a FDistanceCurve variable on the animation blueprint. 
+	 * @param PlayRateAdjustmentCurve - Optional curve that can be used to clamp how much the animation can be advanced in an update (to avoid play back being
+	 *	too fast or too slow). The X axis in the curve should be the resulting play rate from distance matching. The Y axis should be the clamped play rate.
+	 * @param PlayRateClamp - Only used if PlayRateAdjustmentCurve is not set. A simple clamp instead of a curve.
+	 */
+	UFUNCTION(BlueprintCallable, Category = "Distance Matching", meta=(BlueprintThreadSafe))
+	static FSequenceEvaluatorReference AdvanceTimeByDistanceMatching(const FAnimUpdateContext& UpdateContext, const FSequenceEvaluatorReference& SequenceEvaluator,
+		float DistanceTraveled, const FDistanceCurve& CachedDistanceCurve, const UCurveFloat* PlayRateAdjustmentCurve, FVector2D PlayRateClamp = FVector2D(0.75f, 1.25f));
+
+	/**
+	 * Set the time of the sequence evaluator to the point in the animation where the distance curve matches the DistanceToTarget input.
+	 * A common use case is to achieve stops without foot sliding by, each frame, selecting the point in the animation that matches the distance the character has remaining until it stops.
+	 * Note that because this technique sets the time of the animation by distance remaining, it doesn't respect phase of any previous animation (e.g. from a jog cycle).
+	 * @param SequenceEvaluator - The sequence evaluator node to operate on.
+	 * @param DistanceToTarget - The distance remaining to a target (e.g. a stop or pivot point).
+	 * @param CachedDistanceCurve - Optimized access to curve data. This will typically be a FDistanceCurve variable on the animation blueprint. 
+	 */
+	UFUNCTION(BlueprintCallable, Category = "Distance Matching", meta=(BlueprintThreadSafe))
+	static FSequenceEvaluatorReference DistanceMatchToTarget(const FSequenceEvaluatorReference& SequenceEvaluator,
+		float DistanceToTarget, const FDistanceCurve& CachedDistanceCurve);
+};
\ No newline at end of file
diff --git a/Engine/Plugins/Animation/AnimationLocomotionLibrary/Source/Runtime/Public/AnimTurnInPlaceLibrary.h b/Engine/Plugins/Animation/AnimationLocomotionLibrary/Source/Runtime/Public/AnimTurnInPlaceLibrary.h
new file mode 100644
--- /dev/null
+++ b/Engine/Plugins/Animation/AnimationLocomotionLibrary/Source/Runtime/Public/AnimTurnInPlaceLibrary.h
@@ -0,0 +1,36 @@
+// Copyright Epic Games, Inc. All Rights Reserved.
+
+#pragma once
+
+#include "AnimTurnInPlaceTypes.h"
+#include "AnimTurnInPlaceLibrary.generated.h"
+
+class UAnimSequence;
+
+/**
+ * Turn In Place functionality maintains an offset between the capsule and the mesh to prevent the animated pose from spinning when the character rotates on the spot.
+ * Once the offset gets large enough, an animation can be played to rotate the mesh closer to the character's facing direction.
+ */
+UCLASS()
+class ANIMATIONLOCOMOTIONLIBRARYRUNTIME_API UAnimTurnInPlaceLibrary : public UBlueprintFunctionLibrary
+{
+	GENERATED_BODY()
+
+public:
+	/**
+	 * Update a Turn in Place state structure based on the rotation of a character's mesh.
+	 * A typical setup is to have a FAnimTurnInPlaceState variable on the animation blueprint and call this function to update it in BlueprintThreadSafeUpdateAnimation.
+	 * The animation blueprint will use a Rotate Root Bone node to apply the RootYawOffset from the state. The locomotion state machine will have a state to play the 
+	 * active turn transition if requested and a state to play the turn recovery if requested. 
+	 * @param DeltaTime - Time since last update.
+	 * @param bAllowTurnInPlace - This should be set to true when character rotation should no longer rotate the mesh (typically during stops and idles).
+	 * @param bHoldYawOffset - This should be true when the root yaw offset should be maintained but no longer updated (e.g. during start animations).
+	 * @param bIsTurnTransitionStateRelevant - True when the animation blueprint is in the turn transition state. UCachedAnimDataLibrary::StateMachine_IsStateRelevant() can be used for this test.
+	 * @param MeshWorldRotation - The current world rotation of the character's mesh component.
+	 * @param AnimSet - Set of turn on spot animations used to turn the mesh to align with the character's facing.
+	 * @param TurnInPlaceState - The turn in place state information to update. This typically is a variable on the animation blueprint.
+	 */
+	UFUNCTION(BlueprintCallable, Category = "Animation Turn In Place", meta = (BlueprintThreadSafe))
+	static void UpdateTurnInPlace(float DeltaTime, bool bAllowTurnInPlace, bool bHoldYawOffset, bool bIsTurnTransitionStateRelevant,
+		const FRotator& MeshWorldRotation, const FAnimTurnInPlaceAnimSet& AnimSet, UPARAM(ref) FAnimTurnInPlaceState& TurnInPlaceState);
+};
\ No newline at end of file
diff --git a/Engine/Plugins/Animation/AnimationLocomotionLibrary/Source/Runtime/Public/AnimTurnInPlaceTypes.h b/Engine/Plugins/Animation/AnimationLocomotionLibrary/Source/Runtime/Public/AnimTurnInPlaceTypes.h
new file mode 100644
--- /dev/null
+++ b/Engine/Plugins/Animation/AnimationLocomotionLibrary/Source/Runtime/Public/AnimTurnInPlaceTypes.h
@@ -0,0 +1,135 @@
+// Copyright Epic Games, Inc. All Rights Reserved.
+
+#pragma once
+
+#include "Animation/CachedAnimData.h"
+#include "Engine/SpringInterpolator.h"
+#include "AnimTurnInPlaceTypes.generated.h"
+
+class UAnimSequence;
+
+/**
+ * Settings for an animation of a character turning on the spot.
+ */
+USTRUCT(BlueprintType)
+struct FAnimTurnTransition
+{
+	GENERATED_BODY()
+
+public:
+
+	/** The animation to play. */
+	UPROPERTY(EditAnywhere, Category = "TurnTransition")
+	TObjectPtr<UAnimSequence> Anim = nullptr;
+
+	/** Once the offset between the character and mesh is large enough to trigger this transition, wait this time before triggering the animation. */
+	UPROPERTY(EditAnywhere, Category = "TurnTransition")
+	float DelayBeforeTrigger = 0.f;
+};
+
+/**
+ * Set of turn on the spot animations.
+ */
+USTRUCT(BlueprintType)
+struct FAnimTurnInPlaceAnimSet
+{
+	GENERATED_BODY()
+
+public:
+	/**
+	 * Look for the best transition that will result in a smaller offset.
+	 * E.g. If YawOffset is 60 degrees, a transition that rotates 90 degrees will result in a 30 degree offset.
+	 */
+	int32 FindBestTurnTransitionIndex(float YawOffset) const;
+
+	/** Potential animations to play. */
+	UPROPERTY(EditAnywhere, Category = "Turns")
+	TArray<FAnimTurnTransition> TurnTransitions;
+
+	/** The offset between the capsule and mesh has to be bigger than the animation turn angle plus this deadzone to trigger an animation. */
+	UPROPERTY(EditAnywhere, Category = "Turns")
+	float TurnDeadZoneAngle = 0.f;
+};
+
+/**
+ * State that needs to be tracked for triggering turn in place animations and for maintaining an offset between the capsule and mesh. 
+ */
+USTRUCT(BlueprintType)
+struct FAnimTurnInPlaceState
+{
+	GENERATED_BODY()
+
+public:
+
+	FAnimTurnInPlaceState()
+	{
+		YawOffsetInterpolator.SetDefaultSpringConstants(30.f);
+	}
+
+	/**
+	 * See UAnimTurnInPlaceLibrary::UpdateTurnInPlace() for parameter documentation.
+	 */ 
+	void Update(float DeltaTime, bool bAllowTurnInPlace, bool bHoldYawOffset, bool bIsTurnTransitionStateRelevant,
+		const FRotator& MeshWorldRotation, const FAnimTurnInPlaceAnimSet& AnimSet, UPARAM(ref) FAnimTurnInPlaceState& TurnInPlaceState);
+
+	/** Animation that is actively rotating the mesh. */
+	UPROPERTY(BlueprintReadWrite, Category = "Turns")
+	TObjectPtr<UAnimSequence> ActiveTurnAnim = nullptr;
+
+	/** When there's no longer any rotation left in the ActiveTurnAnim, it switches to being the recovery animation to finish playing.
+	    The recovery animation can be interrupted by a new turn animation. They are tracked separately so they can be cross-fade blended together. */
+	UPROPERTY(BlueprintReadWrite, Category = "Turns")
+	TObjectPtr<UAnimSequence> TurnRecoveryAnim = nullptr;
+
+	/** The desired offset between the capsule and the mesh. */
+	UPROPERTY(BlueprintReadWrite, Category = "Root Offset")
+	float RootYawOffset = 0.f;
+
+	/** The inverse of the root yaw offset. This is useful for having an aim offset to keep the upper body looking in the same direction
+		as the character while the lower body stays planted. */
+	UPROPERTY(BlueprintReadWrite, Category = "Root Offset")
+	float RootYawOffsetInverse = 0.f;
+
+	/** The current time of the active turn animation. */
+	UPROPERTY(BlueprintReadWrite, Category = "Turns")
+	float ActiveTurnAnimTime = 0.f;
+
+	/** The time that the recovery animation should start playing at. This will be set to the active turn animation's time when it switches
+		to being the recovery animation. */
+	UPROPERTY(BlueprintReadWrite, Category = "Turns")
+	float TurnRecoveryAnimStartTime = 0.f;
+
+	/** Flag that the animation blueprint can use to trigger a turn in place transition state.*/
+	UPROPERTY(BlueprintReadWrite, Category = "Turns")
+	bool bTurnTransitionRequested = false;
+
+	/** Flag that the animation blueprint can use to trigger a turn in place recovery state. */
+	UPROPERTY(BlueprintReadWrite, Category = "Turns")
+	bool bTurnRecoveryRequested = false;
+
+private:
+	/**
+	 * Advance the turn transition animation and back out any rotation that the animation plays from the root yaw offset.
+	 */ 
+	void UpdateActiveTurnTransition(float DeltaTime);
+	
+	/**
+	 * Check to see if a turn transition should trigger.
+	 */ 
+	void UpdateTurnTransitionTrigger(float DeltaTime, const FAnimTurnInPlaceAnimSet& AnimSet);
+
+	/** If the character starts moving while there's a root yaw offset, the offset will be blended out with this interpolator. */
+	UPROPERTY(EditAnywhere, Category = "Root Offset", meta = (AllowPrivateAccess = "true"))
+	FFloatRK4SpringInterpolator YawOffsetInterpolator;
+
+	/** The turn transition animation that is desired to play, but is still waiting for its trigger delay to finish. */
+	UPROPERTY(transient)
+	TObjectPtr<UAnimSequence> PendingTurnAnim = nullptr;
+
+	/** How long the system has been waiting to trigger PendingTurnAnim. */
+	float PendingTurnDelayCounter = 0.f;
+
+	/** World yaw of the mesh component. This is tracked to detect deltas in rotation between updates.
+		The mesh component's yaw is used, rather than the capsule, to account for mesh smoothing on simulated proxies. */
+	float MeshWorldYaw = FLT_MAX;
+};
\ No newline at end of file
diff --git a/Engine/Plugins/Animation/AnimationLocomotionLibrary/Source/Runtime/Public/AnimationLocomotionLibraryModule.h b/Engine/Plugins/Animation/AnimationLocomotionLibrary/Source/Runtime/Public/AnimationLocomotionLibraryModule.h
new file mode 100644
--- /dev/null
+++ b/Engine/Plugins/Animation/AnimationLocomotionLibrary/Source/Runtime/Public/AnimationLocomotionLibraryModule.h
@@ -0,0 +1,15 @@
+// Copyright Epic Games, Inc. All Rights Reserved.
+
+#pragma once
+
+#include "CoreMinimal.h"
+#include "Modules/ModuleManager.h"
+
+class FAnimationLocomotionLibraryModule : public IModuleInterface
+{
+public:
+
+	/** IModuleInterface implementation */
+	virtual void StartupModule() override;
+	virtual void ShutdownModule() override;
+};
\ No newline at end of file
diff --git a/Engine/Source/Runtime/AnimGraphRuntime/Private/SequenceEvaluatorLibrary.cpp b/Engine/Source/Runtime/AnimGraphRuntime/Private/SequenceEvaluatorLibrary.cpp
--- a/Engine/Source/Runtime/AnimGraphRuntime/Private/SequenceEvaluatorLibrary.cpp
+++ b/Engine/Source/Runtime/AnimGraphRuntime/Private/SequenceEvaluatorLibrary.cpp
@@ -27,6 +27,31 @@ FSequenceEvaluatorReference USequenceEvaluatorLibrary::SetExplicitTime(const FSe
 	return SequenceEvaluator;
 }
 
+FSequenceEvaluatorReference USequenceEvaluatorLibrary::AdvanceTime(const FAnimUpdateContext& UpdateContext, const FSequenceEvaluatorReference& SequenceEvaluator, float PlayRate)
+{
+	SequenceEvaluator.CallAnimNodeFunction<FAnimNode_SequenceEvaluator>(
+		TEXT("AdvanceTime"),
+		[&UpdateContext, PlayRate](FAnimNode_SequenceEvaluator& InSequenceEvaluator)
+		{
+			if (const FAnimationUpdateContext* AnimationUpdateContext = UpdateContext.GetContext())
+			{
+				float ExplicitTime = InSequenceEvaluator.GetExplicitTime();
+				FAnimationRuntime::AdvanceTime(InSequenceEvaluator.GetShouldLoop(), AnimationUpdateContext->GetDeltaTime() * PlayRate, ExplicitTime, InSequenceEvaluator.GetCurrentAssetLength());
+
+				if (!InSequenceEvaluator.SetExplicitTime(ExplicitTime))
+				{
+					UE_LOG(LogSequenceEvaluatorLibrary, Warning, TEXT("Could not advance time on sequence evaluator, ExplicitTime is not dynamic. Set it as Always Dynamic."));
+				}
+			}
+			else
+			{
+				UE_LOG(LogSequenceEvaluatorLibrary, Warning, TEXT("AdvanceTime called with invalid context"));
+			}
+		});
+
+	return SequenceEvaluator;
+}
+
 FSequenceEvaluatorReference USequenceEvaluatorLibrary::SetSequence(const FSequenceEvaluatorReference& SequenceEvaluator, UAnimSequenceBase* Sequence)
 {
 	SequenceEvaluator.CallAnimNodeFunction<FAnimNode_SequenceEvaluator>(
diff --git a/Engine/Source/Runtime/AnimGraphRuntime/Public/SequenceEvaluatorLibrary.h b/Engine/Source/Runtime/AnimGraphRuntime/Public/SequenceEvaluatorLibrary.h
--- a/Engine/Source/Runtime/AnimGraphRuntime/Public/SequenceEvaluatorLibrary.h
+++ b/Engine/Source/Runtime/AnimGraphRuntime/Public/SequenceEvaluatorLibrary.h
@@ -44,6 +44,10 @@ public:
 	UFUNCTION(BlueprintCallable, Category = "Sequence Evaluator", meta=(BlueprintThreadSafe))
 	static FSequenceEvaluatorReference SetExplicitTime(const FSequenceEvaluatorReference& SequenceEvaluator, float Time);
 
+	/** Advance the current accumulated time of the sequence evaluator */
+	UFUNCTION(BlueprintCallable, Category = "Sequence Evaluator", meta = (BlueprintThreadSafe))
+	static FSequenceEvaluatorReference AdvanceTime(const FAnimUpdateContext& UpdateContext, const FSequenceEvaluatorReference& SequenceEvaluator, float PlayRate = 1.0f);
+
 	/** Set the current sequence of the sequence evaluator */
 	UFUNCTION(BlueprintCallable, Category = "Sequence Evaluator", meta=(BlueprintThreadSafe))
 	static FSequenceEvaluatorReference SetSequence(const FSequenceEvaluatorReference& SequenceEvaluator, UAnimSequenceBase* Sequence);
diff --git a/Engine/Source/Runtime/Engine/Private/Animation/AnimCurveTypes.cpp b/Engine/Source/Runtime/Engine/Private/Animation/AnimCurveTypes.cpp
--- a/Engine/Source/Runtime/Engine/Private/Animation/AnimCurveTypes.cpp
+++ b/Engine/Source/Runtime/Engine/Private/Animation/AnimCurveTypes.cpp
@@ -354,12 +354,12 @@ FVectorCurve* FTransformCurve::GetVectorCurveByIndex(int32 Index)
 ////////////////////////////////////////////////////
 //  FCachedFloatCurve
 
-bool FCachedFloatCurve::IsValid(UAnimSequenceBase* InAnimSequence) const
+bool FCachedFloatCurve::IsValid(const UAnimSequenceBase* InAnimSequence) const
 {
 	return ((CurveName != NAME_None) && (GetFloatCurve(InAnimSequence) != nullptr));
 }
 
-float FCachedFloatCurve::GetValueAtPosition(UAnimSequenceBase* InAnimSequence, const float& InPosition) const
+float FCachedFloatCurve::GetValueAtPosition(const UAnimSequenceBase* InAnimSequence, const float& InPosition) const
 {
 	if (const FFloatCurve* DistanceCurve = GetFloatCurve(InAnimSequence))
 	{
@@ -369,9 +369,9 @@ float FCachedFloatCurve::GetValueAtPosition(UAnimSequenceBase* InAnimSequence, c
 	return 0.f;
 }
 
-USkeleton::AnimCurveUID FCachedFloatCurve::GetAnimCurveUID(UAnimSequenceBase* InAnimSequence) const
+USkeleton::AnimCurveUID FCachedFloatCurve::GetAnimCurveUID(const UAnimSequenceBase* InAnimSequence) const
 {
-	if (CachedUID == SmartName::MaxUID && InAnimSequence)
+	if (CurveName != CachedCurveName && InAnimSequence)
 	{
 		if (const USkeleton* Skeleton = InAnimSequence->GetSkeleton())
 		{
@@ -379,6 +379,7 @@ USkeleton::AnimCurveUID FCachedFloatCurve::GetAnimCurveUID(UAnimSequenceBase* In
 			if (CurveNameMapping)
 			{
 				CachedUID = CurveNameMapping->FindUID(CurveName);
+				CachedCurveName = CurveName;
 			}
 		}
 	}
@@ -386,7 +387,7 @@ USkeleton::AnimCurveUID FCachedFloatCurve::GetAnimCurveUID(UAnimSequenceBase* In
 	return CachedUID;
 }
 
-const FFloatCurve* FCachedFloatCurve::GetFloatCurve(UAnimSequenceBase* InAnimSequence) const
+const FFloatCurve* FCachedFloatCurve::GetFloatCurve(const UAnimSequenceBase* InAnimSequence) const
 {
 	if (InAnimSequence)
 	{
@@ -404,7 +405,7 @@ const FFloatCurve* FCachedFloatCurve::GetFloatCurve(UAnimSequenceBase* InAnimSeq
 ////////////////////////////////////////////////////
 //  FDistanceCurve
 
-float FDistanceCurve::GetDistanceRange(UAnimSequenceBase* InAnimSequence) const
+float FDistanceCurve::GetDistanceRange(const UAnimSequenceBase* InAnimSequence) const
 {
 	if (const FFloatCurve* DistanceCurve = GetFloatCurve(InAnimSequence))
 	{
@@ -416,7 +417,7 @@ float FDistanceCurve::GetDistanceRange(UAnimSequenceBase* InAnimSequence) const
 	return 0.f;
 }
 
-float FDistanceCurve::GetAnimPositionFromDistance(UAnimSequenceBase* InAnimSequence, const float& InDistance)
+float FDistanceCurve::GetAnimPositionFromDistance(const UAnimSequenceBase* InAnimSequence, const float& InDistance) const
 {
 	if (const FFloatCurve* DistanceCurve = GetFloatCurve(InAnimSequence))
 	{
diff --git a/Engine/Source/Runtime/Engine/Public/Animation/AnimCurveTypes.h b/Engine/Source/Runtime/Engine/Public/Animation/AnimCurveTypes.h
--- a/Engine/Source/Runtime/Engine/Public/Animation/AnimCurveTypes.h
+++ b/Engine/Source/Runtime/Engine/Public/Animation/AnimCurveTypes.h
@@ -257,23 +257,24 @@ struct ENGINE_API FCachedFloatCurve
 	GENERATED_USTRUCT_BODY()
 
 public:
-	UPROPERTY(EditAnywhere, Category = "Curve Settings")
-		FName CurveName;
+	UPROPERTY(EditAnywhere, BlueprintReadWrite, Category = "Curve Settings")
+	FName CurveName;
 
 private:
 	mutable USkeleton::AnimCurveUID CachedUID;
+	mutable FName CachedCurveName;
 
 public:
 	FCachedFloatCurve()
 		: CachedUID(SmartName::MaxUID)
 	{}
 
-	bool IsValid(UAnimSequenceBase* InAnimSequence) const;
-	float GetValueAtPosition(UAnimSequenceBase* InAnimSequence, const float& InPosition) const;
-	const FFloatCurve* GetFloatCurve(UAnimSequenceBase* InAnimSequence) const;
+	bool IsValid(const UAnimSequenceBase* InAnimSequence) const;
+	float GetValueAtPosition(const UAnimSequenceBase* InAnimSequence, const float& InPosition) const;
+	const FFloatCurve* GetFloatCurve(const UAnimSequenceBase* InAnimSequence) const;
 
 protected:
-	USkeleton::AnimCurveUID GetAnimCurveUID(UAnimSequenceBase* InAnimSequence) const;
+	USkeleton::AnimCurveUID GetAnimCurveUID(const UAnimSequenceBase* InAnimSequence) const;
 };
 
 USTRUCT(BlueprintType)
@@ -290,8 +291,8 @@ public:
 		: Super(Other)
 	{}
 
-	float GetDistanceRange(UAnimSequenceBase* InAnimSequence) const;
-	float GetAnimPositionFromDistance(UAnimSequenceBase* InAnimSequence, const float& InDistance);
+	float GetDistanceRange(const UAnimSequenceBase* InAnimSequence) const;
+	float GetAnimPositionFromDistance(const UAnimSequenceBase* InAnimSequence, const float& InDistance) const;
 };
 
 /**
