Animation fixes for non-unity errors

diff --git a/Engine/Plugins/Animation/AnimationLocomotionLibrary/Source/Runtime/Private/AnimDistanceMatchingLibrary.cpp b/Engine/Plugins/Animation/AnimationLocomotionLibrary/Source/Runtime/Private/AnimDistanceMatchingLibrary.cpp
--- a/Engine/Plugins/Animation/AnimationLocomotionLibrary/Source/Runtime/Private/AnimDistanceMatchingLibrary.cpp
+++ b/Engine/Plugins/Animation/AnimationLocomotionLibrary/Source/Runtime/Private/AnimDistanceMatchingLibrary.cpp
@@ -3,6 +3,7 @@
 #include "AnimDistanceMatchingLibrary.h"
 #include "Animation/AnimSequence.h"
 #include "AnimNodes/AnimNode_SequenceEvaluator.h"
+#include "Curves/CurveFloat.h"
 
 DEFINE_LOG_CATEGORY_STATIC(LogAnimDistanceMatchingLibrary, Verbose, All);
 
diff --git a/Engine/Plugins/Animation/AnimationLocomotionLibrary/Source/Runtime/Private/AnimTurnInPlaceTypes.cpp b/Engine/Plugins/Animation/AnimationLocomotionLibrary/Source/Runtime/Private/AnimTurnInPlaceTypes.cpp
--- a/Engine/Plugins/Animation/AnimationLocomotionLibrary/Source/Runtime/Private/AnimTurnInPlaceTypes.cpp
+++ b/Engine/Plugins/Animation/AnimationLocomotionLibrary/Source/Runtime/Private/AnimTurnInPlaceTypes.cpp
@@ -1,6 +1,7 @@
 // Copyright Epic Games, Inc. All Rights Reserved.
 
 #include "AnimTurnInPlaceTypes.h"
+#include "Animation/AnimSequence.h"
 
 DEFINE_LOG_CATEGORY_STATIC(LogAnimTurnInPlaceLibrary, Verbose, All);
 
diff --git a/Engine/Plugins/Animation/AnimationLocomotionLibrary/Source/Runtime/Public/AnimCharacterMovementLibrary.h b/Engine/Plugins/Animation/AnimationLocomotionLibrary/Source/Runtime/Public/AnimCharacterMovementLibrary.h
--- a/Engine/Plugins/Animation/AnimationLocomotionLibrary/Source/Runtime/Public/AnimCharacterMovementLibrary.h
+++ b/Engine/Plugins/Animation/AnimationLocomotionLibrary/Source/Runtime/Public/AnimCharacterMovementLibrary.h
@@ -3,6 +3,8 @@
 #pragma once
 
 #include "AnimCharacterMovementTypes.h"
+#include "Kismet/BlueprintFunctionLibrary.h"
+
 #include "AnimCharacterMovementLibrary.generated.h"
 
 class UAnimSequence;
diff --git a/Engine/Plugins/Experimental/Animation/MotionTrajectory/Source/MotionTrajectory/Private/MotionTrajectory.cpp b/Engine/Plugins/Experimental/Animation/MotionTrajectory/Source/MotionTrajectory/Private/MotionTrajectory.cpp
--- a/Engine/Plugins/Experimental/Animation/MotionTrajectory/Source/MotionTrajectory/Private/MotionTrajectory.cpp
+++ b/Engine/Plugins/Experimental/Animation/MotionTrajectory/Source/MotionTrajectory/Private/MotionTrajectory.cpp
@@ -2,6 +2,7 @@
 
 #include "MotionTrajectory.h"
 
+#include "Engine/World.h"
 #include "GameFramework/Pawn.h"
 #include "Kismet/KismetMathLibrary.h"
 #include "Kismet/KismetSystemLibrary.h"
