Fixed a packaging issue causing AnimationLocomotion runtime module to not be found in QAGame

diff --git a/Engine/Plugins/Animation/AnimationLocomotionLibrary/Source/Runtime/Private/AnimationLocomotionLibraryModule.cpp b/Engine/Plugins/Animation/AnimationLocomotionLibrary/Source/Runtime/Private/AnimationLocomotionLibraryModule.cpp
--- a/Engine/Plugins/Animation/AnimationLocomotionLibrary/Source/Runtime/Private/AnimationLocomotionLibraryModule.cpp
+++ b/Engine/Plugins/Animation/AnimationLocomotionLibrary/Source/Runtime/Private/AnimationLocomotionLibraryModule.cpp
@@ -4,16 +4,16 @@
 
 #define LOCTEXT_NAMESPACE "AnimationLocomotionModule"
 
-void FAnimationLocomotionLibraryModule::StartupModule()
+void FAnimationLocomotionLibraryRuntimeModule::StartupModule()
 {
 
 }
 
-void FAnimationLocomotionLibraryModule::ShutdownModule()
+void FAnimationLocomotionLibraryRuntimeModule::ShutdownModule()
 {
 
 }
 
 #undef LOCTEXT_NAMESPACE
 
-IMPLEMENT_MODULE(FAnimationLocomotionLibraryModule, AnimationLocomotionLibrary)
\ No newline at end of file
+IMPLEMENT_MODULE(FAnimationLocomotionLibraryRuntimeModule, AnimationLocomotionLibraryRuntime)
\ No newline at end of file
diff --git a/Engine/Plugins/Animation/AnimationLocomotionLibrary/Source/Runtime/Public/AnimationLocomotionLibraryModule.h b/Engine/Plugins/Animation/AnimationLocomotionLibrary/Source/Runtime/Public/AnimationLocomotionLibraryModule.h
--- a/Engine/Plugins/Animation/AnimationLocomotionLibrary/Source/Runtime/Public/AnimationLocomotionLibraryModule.h
+++ b/Engine/Plugins/Animation/AnimationLocomotionLibrary/Source/Runtime/Public/AnimationLocomotionLibraryModule.h
@@ -5,7 +5,7 @@
 #include "CoreMinimal.h"
 #include "Modules/ModuleManager.h"
 
-class FAnimationLocomotionLibraryModule : public IModuleInterface
+class FAnimationLocomotionLibraryRuntimeModule : public IModuleInterface
 {
 public:
 
