Added setting to specify the axes used to calculate the distance value for the distance curve modifier
#rnx

diff --git a/Engine/Plugins/Animation/AnimationLocomotionLibrary/Source/Editor/Private/DistanceCurveModifier.cpp b/Engine/Plugins/Animation/AnimationLocomotionLibrary/Source/Editor/Private/DistanceCurveModifier.cpp
--- a/Engine/Plugins/Animation/AnimationLocomotionLibrary/Source/Editor/Private/DistanceCurveModifier.cpp
+++ b/Engine/Plugins/Animation/AnimationLocomotionLibrary/Source/Editor/Private/DistanceCurveModifier.cpp
@@ -40,7 +40,7 @@ void UDistanceCurveModifier::OnApply_Implementation(UAnimSequence* Animation)
 
 		const bool bAllowLooping = false;
 		const FVector RootMotionTranslation = Animation->ExtractRootMotion(Time, SampleInterval, bAllowLooping).GetTranslation();
-		const float RootMotionSpeedSq = RootMotionTranslation.SizeSquared2D() / SampleInterval;
+		const float RootMotionSpeedSq = CalculateMagnitudeSq(RootMotionTranslation, Axis) / SampleInterval;
 
 		if (RootMotionSpeedSq < MinSpeedSq)
 		{
@@ -61,7 +61,7 @@ void UDistanceCurveModifier::OnApply_Implementation(UAnimSequence* Animation)
 		const float ValueSign = (Time < TimeOfMinSpeed) ? -1.0f : 1.0f;
 
 		const FVector RootMotionTranslation = Animation->ExtractRootMotionFromRange(TimeOfMinSpeed, Time).GetTranslation();
-		UAnimationBlueprintLibrary::AddFloatCurveKey(Animation, CurveName, Time, ValueSign * RootMotionTranslation.Size2D());
+		UAnimationBlueprintLibrary::AddFloatCurveKey(Animation, CurveName, Time, ValueSign * CalculateMagnitude(RootMotionTranslation, Axis));
 	}
 }
 
@@ -69,4 +69,32 @@ void UDistanceCurveModifier::OnRevert_Implementation(UAnimSequence* Animation)
 {
 	const bool bRemoveNameFromSkeleton = false;
 	UAnimationBlueprintLibrary::RemoveCurve(Animation, CurveName, bRemoveNameFromSkeleton);
-}
\ No newline at end of file
+}
+
+float UDistanceCurveModifier::CalculateMagnitude(const FVector& Vector, EDistanceCurve_Axis Axis)
+{
+	switch (Axis)
+	{
+	case EDistanceCurve_Axis::X:		return FMath::Abs(Vector.X); break;
+	case EDistanceCurve_Axis::Y:		return FMath::Abs(Vector.Y); break;
+	case EDistanceCurve_Axis::Z:		return FMath::Abs(Vector.Z); break;
+	default: return FMath::Sqrt(CalculateMagnitudeSq(Vector, Axis)); break;
+	}
+}
+
+float UDistanceCurveModifier::CalculateMagnitudeSq(const FVector& Vector, EDistanceCurve_Axis Axis)
+{
+	switch (Axis)
+	{
+	case EDistanceCurve_Axis::X:		return FMath::Square(FMath::Abs(Vector.X)); break;
+	case EDistanceCurve_Axis::Y:		return FMath::Square(FMath::Abs(Vector.Y)); break;
+	case EDistanceCurve_Axis::Z:		return FMath::Square(FMath::Abs(Vector.Z)); break;
+	case EDistanceCurve_Axis::XY:		return Vector.X * Vector.X + Vector.Y * Vector.Y; break;
+	case EDistanceCurve_Axis::XZ:		return Vector.X * Vector.X + Vector.Z * Vector.Z; break;
+	case EDistanceCurve_Axis::YZ:		return Vector.Y * Vector.Y + Vector.Z * Vector.Z; break;
+	case EDistanceCurve_Axis::XYZ:		return Vector.X * Vector.X + Vector.Y * Vector.Y + Vector.Z * Vector.Z; break;
+	default: check(false); break;
+	}
+
+	return 0.f;
+}
diff --git a/Engine/Plugins/Animation/AnimationLocomotionLibrary/Source/Editor/Public/DistanceCurveModifier.h b/Engine/Plugins/Animation/AnimationLocomotionLibrary/Source/Editor/Public/DistanceCurveModifier.h
--- a/Engine/Plugins/Animation/AnimationLocomotionLibrary/Source/Editor/Public/DistanceCurveModifier.h
+++ b/Engine/Plugins/Animation/AnimationLocomotionLibrary/Source/Editor/Public/DistanceCurveModifier.h
@@ -6,6 +6,20 @@
 #include "AnimationModifier.h"
 #include "DistanceCurveModifier.generated.h"
 
+// @todo: Consolidate with EMotionExtractor_Axis
+/** Axes to calculate the distance value from */
+UENUM(BlueprintType)
+enum class EDistanceCurve_Axis : uint8
+{
+	X,
+	Y,
+	Z,
+	XY,
+	XZ,
+	YZ,
+	XYZ
+};
+
 /** Extracts traveling distance information from root motion and bakes it to a curve.
  * A negative value indicates distance remaining to a stop or pivot point.
  * A positive value indicates distance traveled from a start point or from the beginning of the clip.
@@ -29,6 +43,16 @@ public:
 	UPROPERTY(EditAnywhere, BlueprintReadWrite, Category = Settings)
 	float StopSpeedThreshold = 5.0f;
 
+	/** Axes to calculate the distance value from. */
+	UPROPERTY(EditAnywhere, BlueprintReadWrite, Category = Settings)
+	EDistanceCurve_Axis Axis = EDistanceCurve_Axis::XY;
+
 	virtual void OnApply_Implementation(UAnimSequence* Animation) override;
 	virtual void OnRevert_Implementation(UAnimSequence* Animation) override;
+
+private:
+
+	/** Helper functions to calculate the magnitude of a vector only considering a specific axis or axes */
+	static float CalculateMagnitude(const FVector& Vector, EDistanceCurve_Axis Axis);
+	static float CalculateMagnitudeSq(const FVector& Vector, EDistanceCurve_Axis Axis);
 };
\ No newline at end of file
