Added option to specify the last frame as the root motion stop of a distance matched anim

diff --git a/Engine/Plugins/Animation/AnimationLocomotionLibrary/Source/Editor/Private/DistanceCurveModifier.cpp b/Engine/Plugins/Animation/AnimationLocomotionLibrary/Source/Editor/Private/DistanceCurveModifier.cpp
--- a/Engine/Plugins/Animation/AnimationLocomotionLibrary/Source/Editor/Private/DistanceCurveModifier.cpp
+++ b/Engine/Plugins/Animation/AnimationLocomotionLibrary/Source/Editor/Private/DistanceCurveModifier.cpp
@@ -26,26 +26,36 @@ void UDistanceCurveModifier::OnApply_Implementation(UAnimSequence* Animation)
 	UAnimationBlueprintLibrary::AddCurve(Animation, CurveName, ERawCurveTrackTypes::RCT_Float, bMetaDataCurve);
 
 	const float AnimLength = Animation->GetPlayLength();
+	float SampleInterval;
+	int32 NumSteps;
+	float TimeOfMinSpeed;
 
-	// Perform a high resolution search to find the sample point with minimum speed.
-	
-	float TimeOfMinSpeed = 0.f;
-	float MinSpeedSq = FMath::Square(StopSpeedThreshold);
-
-	float SampleInterval = 1.f / 120.f;
-	int32 NumSteps = AnimLength / SampleInterval;
-	for (int32 Step = 0; Step < NumSteps; ++Step)
+	if(bStopAtEnd)
+	{ 
+		TimeOfMinSpeed = AnimLength;
+	}
+	else
 	{
-		const float Time = Step * SampleInterval;
+		// Perform a high resolution search to find the sample point with minimum speed.
+		
+		TimeOfMinSpeed = 0.f;
+		float MinSpeedSq = FMath::Square(StopSpeedThreshold);
+
+		SampleInterval = 1.f / 120.f;
+		NumSteps = AnimLength / SampleInterval;
+		for (int32 Step = 0; Step < NumSteps; ++Step)
+		{
+			const float Time = Step * SampleInterval;
 
-		const bool bAllowLooping = false;
-		const FVector RootMotionTranslation = Animation->ExtractRootMotion(Time, SampleInterval, bAllowLooping).GetTranslation();
-		const float RootMotionSpeedSq = CalculateMagnitudeSq(RootMotionTranslation, Axis) / SampleInterval;
+			const bool bAllowLooping = false;
+			const FVector RootMotionTranslation = Animation->ExtractRootMotion(Time, SampleInterval, bAllowLooping).GetTranslation();
+			const float RootMotionSpeedSq = CalculateMagnitudeSq(RootMotionTranslation, Axis) / SampleInterval;
 
-		if (RootMotionSpeedSq < MinSpeedSq)
-		{
-			MinSpeedSq = RootMotionSpeedSq;
-			TimeOfMinSpeed = Time;
+			if (RootMotionSpeedSq < MinSpeedSq)
+			{
+				MinSpeedSq = RootMotionSpeedSq;
+				TimeOfMinSpeed = Time;
+			}
 		}
 	}
 
diff --git a/Engine/Plugins/Animation/AnimationLocomotionLibrary/Source/Editor/Public/DistanceCurveModifier.h b/Engine/Plugins/Animation/AnimationLocomotionLibrary/Source/Editor/Public/DistanceCurveModifier.h
--- a/Engine/Plugins/Animation/AnimationLocomotionLibrary/Source/Editor/Public/DistanceCurveModifier.h
+++ b/Engine/Plugins/Animation/AnimationLocomotionLibrary/Source/Editor/Public/DistanceCurveModifier.h
@@ -40,13 +40,17 @@ public:
 	FName CurveName = "Distance";
 
 	/** Root motion speed must be below this threshold to be considered stopped. */
-	UPROPERTY(EditAnywhere, BlueprintReadWrite, Category = Settings)
+	UPROPERTY(EditAnywhere, BlueprintReadWrite, Category = Settings, meta=(EditCondition="!bStopAtEnd"))
 	float StopSpeedThreshold = 5.0f;
 
 	/** Axes to calculate the distance value from. */
 	UPROPERTY(EditAnywhere, BlueprintReadWrite, Category = Settings)
 	EDistanceCurve_Axis Axis = EDistanceCurve_Axis::XY;
 
+	/** Root motion is considered to be stopped at the clip's end */
+	UPROPERTY(EditAnywhere, BlueprintReadWrite, Category = Settings)
+	bool bStopAtEnd = false;
+
 	virtual void OnApply_Implementation(UAnimSequence* Animation) override;
 	virtual void OnRevert_Implementation(UAnimSequence* Animation) override;
 
