Remove code that's no longer needed from the Animation Locomotion Library plugin. This logic is now being handled in blueprint with Anim Node Functions until (if) we have a better mechanism to make it an engine feature.

diff --git a/Engine/Plugins/Animation/AnimationLocomotionLibrary/Source/Runtime/Private/AnimCharacterMovementLibrary.cpp b/Engine/Plugins/Animation/AnimationLocomotionLibrary/Source/Runtime/Private/AnimCharacterMovementLibrary.cpp
--- a/Engine/Plugins/Animation/AnimationLocomotionLibrary/Source/Runtime/Private/AnimCharacterMovementLibrary.cpp
+++ b/Engine/Plugins/Animation/AnimationLocomotionLibrary/Source/Runtime/Private/AnimCharacterMovementLibrary.cpp
@@ -2,79 +2,6 @@
 
 #include "AnimCharacterMovementLibrary.h"
 
-EAnimCardinalDirection UAnimCharacterMovementLibrary::GetCardinalDirectionFromAngle(EAnimCardinalDirection PreviousCardinalDirection, float AngleInDegrees, float DeadZoneAngle)
-{
-	// Use deadzone offset to favor backpedal on S & W, and favor frontpedal on N & E
-	const float AbsoluteAngle = FMath::Abs(AngleInDegrees);
-	if (PreviousCardinalDirection == EAnimCardinalDirection::North)
-	{
-		if (AbsoluteAngle <= (45.f + DeadZoneAngle + DeadZoneAngle))
-		{
-			return EAnimCardinalDirection::North;
-		}
-		else if (AbsoluteAngle >= 135.f - DeadZoneAngle)
-		{
-			return EAnimCardinalDirection::South;
-		}
-		else if (AngleInDegrees > 0.f)
-		{
-			return EAnimCardinalDirection::East;
-		}
-		return EAnimCardinalDirection::West;
-	}
-	else if (PreviousCardinalDirection == EAnimCardinalDirection::South)
-	{
-		if (AbsoluteAngle <= 45.f + DeadZoneAngle)
-		{
-			return EAnimCardinalDirection::North;
-		}
-		else if (AbsoluteAngle >= (135.f - DeadZoneAngle - DeadZoneAngle))
-		{
-			return EAnimCardinalDirection::South;
-		}
-		else if (AngleInDegrees > 0.f)
-		{
-			return EAnimCardinalDirection::East;
-		}
-		return EAnimCardinalDirection::West;
-	}
-
-	// East and West
-	if (AbsoluteAngle <= (45.f + DeadZoneAngle))
-	{
-		return EAnimCardinalDirection::North;
-	}
-	else if (AbsoluteAngle >= (135.f - DeadZoneAngle))
-	{
-		return EAnimCardinalDirection::South;
-	}
-	else if (AngleInDegrees > 0.f)
-	{
-		return EAnimCardinalDirection::East;
-	}
-	return EAnimCardinalDirection::West;
-}
-
-const UAnimSequence* UAnimCharacterMovementLibrary::SelectAnimForCardinalDirection(EAnimCardinalDirection CardinalDirection, const FCardinalDirectionAnimSet& AnimSet)
-{
-	switch (CardinalDirection)
-	{
-		case EAnimCardinalDirection::North:
-			return AnimSet.NorthAnim;
-		case EAnimCardinalDirection::East:
-			return AnimSet.EastAnim;
-		case EAnimCardinalDirection::South:
-			return AnimSet.SouthAnim;
-		case EAnimCardinalDirection::West:
-			return AnimSet.WestAnim;
-		default:
-			checkNoEntry();
-			return nullptr;
-	}
-
-	return nullptr;
-}
-
 FVector UAnimCharacterMovementLibrary::PredictGroundMovementStopLocation(const FVector& Velocity, 
 	bool bUseSeparateBrakingFriction, float BrakingFriction, float GroundFriction, float BrakingFrictionFactor, float BrakingDecelerationWalking)
 {
diff --git a/Engine/Plugins/Animation/AnimationLocomotionLibrary/Source/Runtime/Private/AnimTurnInPlaceLibrary.cpp b/Engine/Plugins/Animation/AnimationLocomotionLibrary/Source/Runtime/Private/AnimTurnInPlaceLibrary.cpp
deleted file mode 100644
diff --git a/Engine/Plugins/Animation/AnimationLocomotionLibrary/Source/Runtime/Private/AnimTurnInPlaceTypes.cpp b/Engine/Plugins/Animation/AnimationLocomotionLibrary/Source/Runtime/Private/AnimTurnInPlaceTypes.cpp
deleted file mode 100644
diff --git a/Engine/Plugins/Animation/AnimationLocomotionLibrary/Source/Runtime/Public/AnimCharacterMovementLibrary.h b/Engine/Plugins/Animation/AnimationLocomotionLibrary/Source/Runtime/Public/AnimCharacterMovementLibrary.h
--- a/Engine/Plugins/Animation/AnimationLocomotionLibrary/Source/Runtime/Public/AnimCharacterMovementLibrary.h
+++ b/Engine/Plugins/Animation/AnimationLocomotionLibrary/Source/Runtime/Public/AnimCharacterMovementLibrary.h
@@ -2,7 +2,6 @@
 
 #pragma once
 
-#include "AnimCharacterMovementTypes.h"
 #include "Kismet/BlueprintFunctionLibrary.h"
 
 #include "AnimCharacterMovementLibrary.generated.h"
@@ -19,26 +18,6 @@ class ANIMATIONLOCOMOTIONLIBRARYRUNTIME_API UAnimCharacterMovementLibrary : publ
 
 public:
 
-	/**
-	 * Calculate the closest cardinal direction to the direction the character is currently moving.
-	 * @param PreviousCardinalDirection - The cardinal direction from the previous frame. Typically the animation blueprint holds a EAnimCardinalDirection variable.
-	 * @param DirectionAngleInDegrees - The direction that the character is currently moving.
-	 * @param DirectionDeadZoneAngle - Deadzone to prevent flickering between directions at angle boundaries.
-	 * @return The resulting cardinal direction.
-	 */
-	UFUNCTION(BlueprintPure, Category = "Animation Character Movement", meta = (BlueprintThreadSafe))
-	static EAnimCardinalDirection GetCardinalDirectionFromAngle(EAnimCardinalDirection PreviousCardinalDirection, float DirectionAngleInDegrees, float DeadZoneAngle);
-
-	/**
-	 * Select an animation to play based on the cardinal direction calculated by GetCardinalDirectionFromAngle(). For example, this can pick a start animation
-	 * to play based on the character's movement direction.	    
-	 * @param CardinalDirection - The closest cardinal direction to the character's movement direction.
-	 * @param AnimSet - The set of animations to choose from.
-	 * @return The animation to play.
-	 */
-	UFUNCTION(BlueprintPure, Category = "Animation Character Movement", meta = (BlueprintThreadSafe))
-	static const UAnimSequence* SelectAnimForCardinalDirection(EAnimCardinalDirection CardinalDirection, const FCardinalDirectionAnimSet& AnimSet);
-
 	/**
 	 * Predict where the character will stop based on its current movement properties and parameters from the movement component.
 	 * This uses prediction logic that is heavily tied to the UCharacterMovementComponent.
diff --git a/Engine/Plugins/Animation/AnimationLocomotionLibrary/Source/Runtime/Public/AnimCharacterMovementTypes.h b/Engine/Plugins/Animation/AnimationLocomotionLibrary/Source/Runtime/Public/AnimCharacterMovementTypes.h
deleted file mode 100644
diff --git a/Engine/Plugins/Animation/AnimationLocomotionLibrary/Source/Runtime/Public/AnimTurnInPlaceLibrary.h b/Engine/Plugins/Animation/AnimationLocomotionLibrary/Source/Runtime/Public/AnimTurnInPlaceLibrary.h
deleted file mode 100644
diff --git a/Engine/Plugins/Animation/AnimationLocomotionLibrary/Source/Runtime/Public/AnimTurnInPlaceTypes.h b/Engine/Plugins/Animation/AnimationLocomotionLibrary/Source/Runtime/Public/AnimTurnInPlaceTypes.h
deleted file mode 100644
