Removed FDistanceCurve (This was added and removed in UE5 so we're not following the standard deprecation path).
Moved FDistance curve logic to distance matching library

diff --git a/Engine/Plugins/Animation/AnimationLocomotionLibrary/Source/Runtime/Private/AnimDistanceMatchingLibrary.cpp b/Engine/Plugins/Animation/AnimationLocomotionLibrary/Source/Runtime/Private/AnimDistanceMatchingLibrary.cpp
--- a/Engine/Plugins/Animation/AnimationLocomotionLibrary/Source/Runtime/Private/AnimDistanceMatchingLibrary.cpp
+++ b/Engine/Plugins/Animation/AnimationLocomotionLibrary/Source/Runtime/Private/AnimDistanceMatchingLibrary.cpp
@@ -5,86 +5,166 @@
 #include "Animation/AnimSequence.h"
 #include "AnimNodes/AnimNode_SequenceEvaluator.h"
 #include "Curves/CurveFloat.h"
+#include "Animation/AnimCurveCompressionCodec_UniformIndexable.h"
 
 DEFINE_LOG_CATEGORY_STATIC(LogAnimDistanceMatchingLibrary, Verbose, All);
 
-/**
- * Advance from the current time to a new time in the animation that will result in the desired distance traveled by the authored root motion.
- */
-float GetTimeAfterDistanceTraveled(const UAnimSequence* AnimSequence, float CurrentTime, float DistanceTraveled, const FDistanceCurve& CachedDistanceCurve, const bool bAllowLooping)
+namespace UE::Anim::DistanceMatchingUtility
 {
-	float NewTime = CurrentTime;
-	if (AnimSequence != nullptr)
+	static float GetDistanceRange(const UAnimSequenceBase* InAnimSequence, USkeleton::AnimCurveUID CurveUID)
 	{
-		// Avoid infinite loops if the animation doesn't cover any distance.
-		if (!FMath::IsNearlyZero(CachedDistanceCurve.GetDistanceRange(AnimSequence)))
+		FAnimCurveBufferAccess BufferCurveAccess(InAnimSequence, CurveUID);
+		if (BufferCurveAccess.IsValid())
 		{
-			float AccumulatedDistance = 0.f;
+			const int32 NumSamples = BufferCurveAccess.GetNumSamples();
+			if (NumSamples >= 2)
+			{
+				return (BufferCurveAccess.GetValue(NumSamples - 1) - BufferCurveAccess.GetValue(0));
+			}
+		}
+		return 0.f;
+	}
 
-			const float SequenceLength = AnimSequence->GetPlayLength();
-			const float StepTime = 1.f / 30.f;
-			
-			// Distance Matching expects the distance curve on the animation to increase monotonically. If the curve fails to increase in value
-			// after a certain number of iterations, we abandon the algorithm to avoid an infinite loop.
-			const int32 StuckLoopThreshold = 5;
-			int32 StuckLoopCounter = 0;
+	static USkeleton::AnimCurveUID GetCurveUID(const UAnimSequenceBase* InAnimSequence, FName CurveName)
+	{
+		if (const USkeleton* Skeleton = InAnimSequence->GetSkeleton())
+		{
+			const FSmartNameMapping* CurveNameMapping = Skeleton->GetSmartNameContainer(USkeleton::AnimCurveMappingName);
+			if (CurveNameMapping)
+			{
+				return CurveNameMapping->FindUID(CurveName);
+			}
+		}
+
+		return SmartName::MaxUID;
+	}
 
-			// Traverse the distance curve, accumulating animated distance until the desired distance is reached.
-			while ((AccumulatedDistance < DistanceTraveled) && (bAllowLooping || (NewTime + StepTime < SequenceLength)))
+	static float GetAnimPositionFromDistance(const UAnimSequenceBase* InAnimSequence, const float& InDistance, USkeleton::AnimCurveUID CurveUID)
+	{	
+		FAnimCurveBufferAccess BufferCurveAccess(InAnimSequence, CurveUID);
+		if (BufferCurveAccess.IsValid())
+		{
+			const int32 NumKeys = BufferCurveAccess.GetNumSamples();
+			if (NumKeys < 2)
 			{
-				const float CurrentDistance = CachedDistanceCurve.GetValueAtPosition(AnimSequence, NewTime);
-				const float DistanceAfterStep = CachedDistanceCurve.GetValueAtPosition(AnimSequence, NewTime + StepTime);
-				const float AnimationDistanceThisStep = DistanceAfterStep - CurrentDistance;
+				return 0.f;
+			}
 
-				if (!FMath::IsNearlyZero(AnimationDistanceThisStep))
-				{
-					// Keep advancing if the desired distance hasn't been reached.
-					if (AccumulatedDistance + AnimationDistanceThisStep < DistanceTraveled)
-					{
-						FAnimationRuntime::AdvanceTime(bAllowLooping, StepTime, NewTime, SequenceLength);
-						AccumulatedDistance += AnimationDistanceThisStep;
-					}
-					// Once the desired distance is passed, find the approximate time between samples where the distance will be reached.
-					else
-					{
-						const float DistanceAlpha = (DistanceTraveled - AccumulatedDistance) / AnimationDistanceThisStep;
-						FAnimationRuntime::AdvanceTime(bAllowLooping, DistanceAlpha * StepTime, NewTime, SequenceLength);
-						AccumulatedDistance = DistanceTraveled;
-						break;
-					}
+			// Some assumptions: 
+			// - keys have unique values, so for a given value, it maps to a single position on the timeline of the animation.
+			// - key values are sorted in increasing order.
+
+			int32 First = 1;
+			int32 Last = NumKeys - 1;
+			int32 Count = Last - First;
+
+			while (Count > 0)
+			{
+				int32 Step = Count / 2;
+				int32 Middle = First + Step;
 
-					StuckLoopCounter = 0;
+				if (InDistance > BufferCurveAccess.GetValue(Middle))
+				{
+					First = Middle + 1;
+					Count -= Step + 1;
 				}
 				else
 				{
-					++StuckLoopCounter;
-					if (StuckLoopCounter >= StuckLoopThreshold)
+					Count = Step;
+				}
+			}
+
+			const float KeyAValue = BufferCurveAccess.GetValue(First - 1);
+			const float KeyBValue = BufferCurveAccess.GetValue(First);
+			const float Diff = KeyBValue - KeyAValue;
+			const float Alpha = !FMath::IsNearlyZero(Diff) ? ((InDistance - KeyAValue) / Diff) : 0.f;
+
+			const float KeyATime = BufferCurveAccess.GetTime(First - 1);
+			const float KeyBTime = BufferCurveAccess.GetTime(First);
+			return FMath::Lerp(KeyATime, KeyBTime, Alpha);
+		}
+
+		return 0.f;
+	}
+
+	/**
+	* Advance from the current time to a new time in the animation that will result in the desired distance traveled by the authored root motion.
+	*/
+	static float GetTimeAfterDistanceTraveled(const UAnimSequenceBase* AnimSequence, float CurrentTime, float DistanceTraveled, USkeleton::AnimCurveUID CurveUID, const bool bAllowLooping)
+	{
+		float NewTime = CurrentTime;
+		if (AnimSequence != nullptr)
+		{
+			// Avoid infinite loops if the animation doesn't cover any distance.
+			if (!FMath::IsNearlyZero(UE::Anim::DistanceMatchingUtility::GetDistanceRange(AnimSequence, CurveUID)))
+			{
+				float AccumulatedDistance = 0.f;
+
+				const float SequenceLength = AnimSequence->GetPlayLength();
+				const float StepTime = 1.f / 30.f;
+
+				// Distance Matching expects the distance curve on the animation to increase monotonically. If the curve fails to increase in value
+				// after a certain number of iterations, we abandon the algorithm to avoid an infinite loop.
+				const int32 StuckLoopThreshold = 5;
+				int32 StuckLoopCounter = 0;
+
+				// Traverse the distance curve, accumulating animated distance until the desired distance is reached.
+				while ((AccumulatedDistance < DistanceTraveled) && (bAllowLooping || (NewTime + StepTime < SequenceLength)))
+				{
+					const float CurrentDistance = AnimSequence->EvaluateCurveData(CurveUID, NewTime);
+					const float DistanceAfterStep = AnimSequence->EvaluateCurveData(CurveUID, NewTime + StepTime);
+					const float AnimationDistanceThisStep = DistanceAfterStep - CurrentDistance;
+
+					if (!FMath::IsNearlyZero(AnimationDistanceThisStep))
+					{
+						// Keep advancing if the desired distance hasn't been reached.
+						if (AccumulatedDistance + AnimationDistanceThisStep < DistanceTraveled)
+						{
+							FAnimationRuntime::AdvanceTime(bAllowLooping, StepTime, NewTime, SequenceLength);
+							AccumulatedDistance += AnimationDistanceThisStep;
+						}
+						// Once the desired distance is passed, find the approximate time between samples where the distance will be reached.
+						else
+						{
+							const float DistanceAlpha = (DistanceTraveled - AccumulatedDistance) / AnimationDistanceThisStep;
+							FAnimationRuntime::AdvanceTime(bAllowLooping, DistanceAlpha * StepTime, NewTime, SequenceLength);
+							AccumulatedDistance = DistanceTraveled;
+							break;
+						}
+
+						StuckLoopCounter = 0;
+					}
+					else
 					{
-						UE_LOG(LogAnimDistanceMatchingLibrary, Warning, TEXT("Failed to advance any distance after %d loops on anim sequence (%s). Aborting."), StuckLoopThreshold, *GetNameSafe(AnimSequence));
-						break;
+						++StuckLoopCounter;
+						if (StuckLoopCounter >= StuckLoopThreshold)
+						{
+							UE_LOG(LogAnimDistanceMatchingLibrary, Warning, TEXT("Failed to advance any distance after %d loops on anim sequence (%s). Aborting."), StuckLoopThreshold, *GetNameSafe(AnimSequence));
+							break;
+						}
 					}
 				}
 			}
+			else
+			{
+				UE_LOG(LogAnimDistanceMatchingLibrary, Warning, TEXT("Anim sequence (%s) is missing a distance curve or doesn't cover enough distance for GetTimeAfterDistanceTraveled."), *GetNameSafe(AnimSequence));
+			}
 		}
 		else
 		{
-			UE_LOG(LogAnimDistanceMatchingLibrary, Warning, TEXT("Anim sequence (%s) is missing a distance curve or doesn't cover enough distance for GetTimeAfterDistanceTraveled."), *GetNameSafe(AnimSequence));
+			UE_LOG(LogAnimDistanceMatchingLibrary, Warning, TEXT("Invalid AnimSequence passed to GetTimeAfterDistanceTraveled"));
 		}
-	}
-	else
-	{
-		UE_LOG(LogAnimDistanceMatchingLibrary, Warning, TEXT("Invalid AnimSequence passed to GetTimeAfterDistanceTraveled"));
-	}
 
-	return NewTime;
+		return NewTime;
+	}
 }
 
 FSequenceEvaluatorReference UAnimDistanceMatchingLibrary::AdvanceTimeByDistanceMatching(const FAnimUpdateContext& UpdateContext, const FSequenceEvaluatorReference& SequenceEvaluator,
-	float DistanceTraveled, const FDistanceCurve& CachedDistanceCurve, FVector2D PlayRateClamp)
+	float DistanceTraveled, FName DistanceCurveName, FVector2D PlayRateClamp)
 {
 	SequenceEvaluator.CallAnimNodeFunction<FAnimNode_SequenceEvaluator>(
 		TEXT("AdvanceTimeByDistanceMatching"),
-		[&UpdateContext, DistanceTraveled, &CachedDistanceCurve, PlayRateClamp](FAnimNode_SequenceEvaluator& InSequenceEvaluator)
+		[&UpdateContext, DistanceTraveled, DistanceCurveName, PlayRateClamp](FAnimNode_SequenceEvaluator& InSequenceEvaluator)
 		{
 			if (const FAnimationUpdateContext* AnimationUpdateContext = UpdateContext.GetContext())
 			{
@@ -92,13 +172,14 @@ FSequenceEvaluatorReference UAnimDistanceMatchingLibrary::AdvanceTimeByDistanceM
 
 				if (DeltaTime > 0 && DistanceTraveled > 0)
 				{
-					if (const UAnimSequence* AnimSequence = Cast<UAnimSequence>(InSequenceEvaluator.GetSequence()))
+					if (const UAnimSequenceBase* AnimSequence = Cast<UAnimSequence>(InSequenceEvaluator.GetSequence()))
 					{
 						const float CurrentTime = InSequenceEvaluator.GetExplicitTime();
 						const float CurrentAssetLength = InSequenceEvaluator.GetCurrentAssetLength();
 						const bool bAllowLooping = InSequenceEvaluator.GetShouldLoop();
 
-						float TimeAfterDistanceTraveled = GetTimeAfterDistanceTraveled(AnimSequence, CurrentTime, DistanceTraveled, CachedDistanceCurve, bAllowLooping);
+						const USkeleton::AnimCurveUID CurveUID = UE::Anim::DistanceMatchingUtility::GetCurveUID(AnimSequence, DistanceCurveName);
+						float TimeAfterDistanceTraveled = UE::Anim::DistanceMatchingUtility::GetTimeAfterDistanceTraveled(AnimSequence, CurrentTime, DistanceTraveled, CurveUID, bAllowLooping);
 
 						// Calculate the effective playrate that would result from advancing the animation by the distance traveled.
 						// Account for the animation looping.
@@ -139,18 +220,19 @@ FSequenceEvaluatorReference UAnimDistanceMatchingLibrary::AdvanceTimeByDistanceM
 }
 
 FSequenceEvaluatorReference UAnimDistanceMatchingLibrary::DistanceMatchToTarget(const FSequenceEvaluatorReference& SequenceEvaluator,
-	float DistanceToTarget, const FDistanceCurve& CachedDistanceCurve)
+	float DistanceToTarget, FName DistanceCurveName)
 {
 	SequenceEvaluator.CallAnimNodeFunction<FAnimNode_SequenceEvaluator>(
 		TEXT("DistanceMatchToTarget"),
-		[DistanceToTarget, &CachedDistanceCurve](FAnimNode_SequenceEvaluator& InSequenceEvaluator)
+		[DistanceToTarget, DistanceCurveName](FAnimNode_SequenceEvaluator& InSequenceEvaluator)
 		{
-			if (const UAnimSequence* AnimSequence = Cast<UAnimSequence>(InSequenceEvaluator.GetSequence()))
+			if (const UAnimSequenceBase* AnimSequence = Cast<UAnimSequence>(InSequenceEvaluator.GetSequence()))
 			{
-				if (CachedDistanceCurve.IsValid(AnimSequence))
+				const USkeleton::AnimCurveUID CurveUID = UE::Anim::DistanceMatchingUtility::GetCurveUID(AnimSequence, DistanceCurveName);
+				if (AnimSequence->HasCurveData(CurveUID))
 				{
 					// By convention, distance curves store the distance to a target as a negative value.
-					const float NewTime = CachedDistanceCurve.GetAnimPositionFromDistance(AnimSequence, -DistanceToTarget);
+					const float NewTime = UE::Anim::DistanceMatchingUtility::GetAnimPositionFromDistance(AnimSequence, -DistanceToTarget, CurveUID);
 					if (!InSequenceEvaluator.SetExplicitTime(NewTime))
 					{
 						UE_LOG(LogAnimDistanceMatchingLibrary, Warning, TEXT("Could not set explicit time on sequence evaluator, value is not dynamic. Set it as Always Dynamic."));
@@ -158,7 +240,7 @@ FSequenceEvaluatorReference UAnimDistanceMatchingLibrary::DistanceMatchToTarget(
 				}
 				else
 				{
-					UE_LOG(LogAnimDistanceMatchingLibrary, Warning, TEXT("DistanceMatchToTarget called with invalid CachedDistanceCurve or animation (%s) is missing a distance curve."), *GetNameSafe(AnimSequence));
+					UE_LOG(LogAnimDistanceMatchingLibrary, Warning, TEXT("DistanceMatchToTarget called with invalid DistanceCurveName or animation (%s) is missing a distance curve."), *GetNameSafe(AnimSequence));
 				}
 			}
 			else
diff --git a/Engine/Plugins/Animation/AnimationLocomotionLibrary/Source/Runtime/Public/AnimDistanceMatchingLibrary.h b/Engine/Plugins/Animation/AnimationLocomotionLibrary/Source/Runtime/Public/AnimDistanceMatchingLibrary.h
--- a/Engine/Plugins/Animation/AnimationLocomotionLibrary/Source/Runtime/Public/AnimDistanceMatchingLibrary.h
+++ b/Engine/Plugins/Animation/AnimationLocomotionLibrary/Source/Runtime/Public/AnimDistanceMatchingLibrary.h
@@ -27,12 +27,12 @@ public:
 	 * @param UpdateContext - The update context provided in the anim node function.
 	 * @param SequenceEvaluator - The sequence evaluator node to operate on.
 	 * @param DistanceTraveled - The distance traveled by the character since the last animation update.
-	 * @param CachedDistanceCurve - Optimized access to curve data. This will typically be a FDistanceCurve variable on the animation blueprint. 
+	 * @param DistanceCurveName - Name of the curve we want to match 
 	 * @param PlayRateClamp - A clamp on the effective play rate of the animation after distance matching. Set to (0,0) for no clamping.
 	 */
 	UFUNCTION(BlueprintCallable, Category = "Distance Matching", meta=(BlueprintThreadSafe))
 	static FSequenceEvaluatorReference AdvanceTimeByDistanceMatching(const FAnimUpdateContext& UpdateContext, const FSequenceEvaluatorReference& SequenceEvaluator,
-		float DistanceTraveled, const FDistanceCurve& CachedDistanceCurve, FVector2D PlayRateClamp = FVector2D(0.75f, 1.25f));
+		float DistanceTraveled, FName DistanceCurveName, FVector2D PlayRateClamp = FVector2D(0.75f, 1.25f));
 
 	/**
 	 * Set the time of the sequence evaluator to the point in the animation where the distance curve matches the DistanceToTarget input.
@@ -40,11 +40,11 @@ public:
 	 * Note that because this technique sets the time of the animation by distance remaining, it doesn't respect phase of any previous animation (e.g. from a jog cycle).
 	 * @param SequenceEvaluator - The sequence evaluator node to operate on.
 	 * @param DistanceToTarget - The distance remaining to a target (e.g. a stop or pivot point).
-	 * @param CachedDistanceCurve - Optimized access to curve data. This will typically be a FDistanceCurve variable on the animation blueprint. 
+	 * @param DistanceCurveName - Name of the curve we want to match 
 	 */
 	UFUNCTION(BlueprintCallable, Category = "Distance Matching", meta=(BlueprintThreadSafe))
 	static FSequenceEvaluatorReference DistanceMatchToTarget(const FSequenceEvaluatorReference& SequenceEvaluator,
-		float DistanceToTarget, const FDistanceCurve& CachedDistanceCurve);
+		float DistanceToTarget, FName DistanceCurveName);
 
 	/**
 	 * Set the play rate of the sequence player so that the speed of the animation matches in-game movement speed.
diff --git a/Engine/Source/Runtime/Engine/Private/Animation/AnimCurveTypes.cpp b/Engine/Source/Runtime/Engine/Private/Animation/AnimCurveTypes.cpp
--- a/Engine/Source/Runtime/Engine/Private/Animation/AnimCurveTypes.cpp
+++ b/Engine/Source/Runtime/Engine/Private/Animation/AnimCurveTypes.cpp
@@ -398,71 +398,6 @@ const FFloatCurve* FCachedFloatCurve::GetFloatCurve(const UAnimSequenceBase* InA
 	return nullptr;
 }
 
-////////////////////////////////////////////////////
-//  FDistanceCurve
-
-float FDistanceCurve::GetDistanceRange(const UAnimSequenceBase* InAnimSequence) const
-{
-	FAnimCurveBufferAccess BufferCurveAccess(InAnimSequence, GetAnimCurveUID(InAnimSequence));
-	if (BufferCurveAccess.IsValid())
-	{
-		const int32 NumSamples = BufferCurveAccess.GetNumSamples();
-		if (NumSamples >= 2)
-		{
-			return (BufferCurveAccess.GetValue(NumSamples - 1) - BufferCurveAccess.GetValue(0));
-		}
-	}
-	return 0.f;
-}
-
-float FDistanceCurve::GetAnimPositionFromDistance(const UAnimSequenceBase* InAnimSequence, const float& InDistance) const
-{
-	FAnimCurveBufferAccess BufferCurveAccess(InAnimSequence, GetAnimCurveUID(InAnimSequence));
-	if (BufferCurveAccess.IsValid())
-	{
-		const int32 NumKeys = BufferCurveAccess.GetNumSamples();
-		if (NumKeys < 2)
-		{
-			return 0.f;
-		}
-
-		// Some assumptions: 
-		// - keys have unique values, so for a given value, it maps to a single position on the timeline of the animation.
-		// - key values are sorted in increasing order.
-
-		int32 first = 1;
-		int32 last = NumKeys - 1;
-		int32 count = last - first;
-
-		while (count > 0)
-		{
-			int32 step = count / 2;
-			int32 middle = first + step;
-
-			if (InDistance > BufferCurveAccess.GetValue(middle))
-			{
-				first = middle + 1;
-				count -= step + 1;
-			}
-			else
-			{
-				count = step;
-			}
-		}
-
-		const float KeyAValue = BufferCurveAccess.GetValue(first - 1);
-		const float KeyBValue = BufferCurveAccess.GetValue(first);
-		const float Diff = KeyBValue - KeyAValue;
-		const float Alpha = !FMath::IsNearlyZero(Diff) ? ((InDistance - KeyAValue) / Diff) : 0.f;
-
-		const float KeyATime = BufferCurveAccess.GetTime(first - 1);
-		const float KeyBTime = BufferCurveAccess.GetTime(first);
-		return FMath::Lerp(KeyATime, KeyBTime, Alpha);
-	}
-
-	return 0.f;
-}
-
 /////////////////////////////////////////////////////
 // FRawCurveTracks
 
diff --git a/Engine/Source/Runtime/Engine/Public/Animation/AnimCurveTypes.h b/Engine/Source/Runtime/Engine/Public/Animation/AnimCurveTypes.h
--- a/Engine/Source/Runtime/Engine/Public/Animation/AnimCurveTypes.h
+++ b/Engine/Source/Runtime/Engine/Public/Animation/AnimCurveTypes.h
@@ -277,24 +277,6 @@ protected:
 	USkeleton::AnimCurveUID GetAnimCurveUID(const UAnimSequenceBase* InAnimSequence) const;
 };
 
-USTRUCT(BlueprintType)
-struct ENGINE_API FDistanceCurve : public FCachedFloatCurve
-{
-	GENERATED_USTRUCT_BODY()
-
-public:
-	FDistanceCurve()
-		: Super()
-	{}
-
-	FDistanceCurve(const Super& Other)
-		: Super(Other)
-	{}
-
-	float GetDistanceRange(const UAnimSequenceBase* InAnimSequence) const;
-	float GetAnimPositionFromDistance(const UAnimSequenceBase* InAnimSequence, const float& InDistance) const;
-};
-
 /**
 * This is array of curves that run when collecting curves natively 
 */
