Package file can't be found for the AnimationLocomotionLibraryEditor while running project in the Standalone Game mode
#fix set Plugin modules containing AnimationModifiers to UncookedOnly rather than Editor - as references in standalone -Game session will be unresolved

diff --git a/Engine/Plugins/Animation/AnimationLocomotionLibrary/AnimationLocomotionLibrary.uplugin b/Engine/Plugins/Animation/AnimationLocomotionLibrary/AnimationLocomotionLibrary.uplugin
--- a/Engine/Plugins/Animation/AnimationLocomotionLibrary/AnimationLocomotionLibrary.uplugin
+++ b/Engine/Plugins/Animation/AnimationLocomotionLibrary/AnimationLocomotionLibrary.uplugin
@@ -23,7 +23,7 @@
 		},
 		{
 			"Name": "AnimationLocomotionLibraryEditor",
-			"Type": "Editor",
+			"Type": "UncookedOnly",
 			"LoadingPhase": "Default"
 		}
 	]
diff --git a/Engine/Plugins/Animation/AnimationModifierLibrary/AnimationModifierLibrary.uplugin b/Engine/Plugins/Animation/AnimationModifierLibrary/AnimationModifierLibrary.uplugin
--- a/Engine/Plugins/Animation/AnimationModifierLibrary/AnimationModifierLibrary.uplugin
+++ b/Engine/Plugins/Animation/AnimationModifierLibrary/AnimationModifierLibrary.uplugin
@@ -19,7 +19,7 @@
 	[
 		{
 			"Name": "AnimationModifierLibrary",
-			"Type": "Editor",
+			"Type": "UncookedOnly",
 			"LoadingPhase": "Default"
 		}
 	]
