Improve stuck detection in Distance Matching

diff --git a/Engine/Plugins/Animation/AnimationLocomotionLibrary/Source/Runtime/Private/AnimDistanceMatchingLibrary.cpp b/Engine/Plugins/Animation/AnimationLocomotionLibrary/Source/Runtime/Private/AnimDistanceMatchingLibrary.cpp
--- a/Engine/Plugins/Animation/AnimationLocomotionLibrary/Source/Runtime/Private/AnimDistanceMatchingLibrary.cpp
+++ b/Engine/Plugins/Animation/AnimationLocomotionLibrary/Source/Runtime/Private/AnimDistanceMatchingLibrary.cpp
@@ -88,14 +88,10 @@ namespace UE::Anim::DistanceMatchingUtility
 			if (!FMath::IsNearlyZero(UE::Anim::DistanceMatchingUtility::GetDistanceRange(AnimSequence, CurveName)))
 			{
 				float AccumulatedDistance = 0.f;
+				float AccumulatedTime = 0.f;
 
 				const float SequenceLength = AnimSequence->GetPlayLength();
-				const float StepTime = 1.f / 30.f;
-
-				// Distance Matching expects the distance curve on the animation to increase monotonically. If the curve fails to increase in value
-				// after a certain number of iterations, we abandon the algorithm to avoid an infinite loop.
-				const int32 StuckLoopThreshold = 5;
-				int32 StuckLoopCounter = 0;
+				static const float StepTime = 1.f / 30.f;				
 
 				// Traverse the distance curve, accumulating animated distance until the desired distance is reached.
 				while ((AccumulatedDistance < DistanceTraveled) && (bAllowLooping || (NewTime + StepTime < SequenceLength)))
@@ -120,17 +116,20 @@ namespace UE::Anim::DistanceMatchingUtility
 							AccumulatedDistance = DistanceTraveled;
 							break;
 						}
-
-						StuckLoopCounter = 0;
 					}
 					else
 					{
-						++StuckLoopCounter;
-						if (StuckLoopCounter >= StuckLoopThreshold)
-						{
-							UE_LOG(LogAnimDistanceMatchingLibrary, Warning, TEXT("Failed to advance any distance after %d loops on anim sequence (%s). Aborting."), StuckLoopThreshold, *GetNameSafe(AnimSequence));
-							break;
-						}
+						NewTime += StepTime;
+					}
+					
+					AccumulatedTime += StepTime;
+
+					// If the animation doesn't cover enough distance, we abandon the algorithm to avoid an infinite loop.
+					if (AccumulatedTime >= SequenceLength)
+					{
+						UE_LOG(LogAnimDistanceMatchingLibrary, Warning, TEXT("Failed to advance distance of (%.2f) after (%.2f) seconds on anim sequence (%s). Aborting."), 
+							DistanceTraveled, AccumulatedTime, *GetNameSafe(AnimSequence));
+						break;
 					}
 				}
 			}
