Change AnimationLocomotionLibrary and AnimationWarping loading phase to PreDefault

diff --git a/Engine/Plugins/Animation/AnimationLocomotionLibrary/AnimationLocomotionLibrary.uplugin b/Engine/Plugins/Animation/AnimationLocomotionLibrary/AnimationLocomotionLibrary.uplugin
--- a/Engine/Plugins/Animation/AnimationLocomotionLibrary/AnimationLocomotionLibrary.uplugin
+++ b/Engine/Plugins/Animation/AnimationLocomotionLibrary/AnimationLocomotionLibrary.uplugin
@@ -19,12 +19,12 @@
 		{
 			"Name": "AnimationLocomotionLibraryRuntime",
 			"Type": "Runtime",
-			"LoadingPhase": "Default"
+			"LoadingPhase": "PreDefault"
 		},
 		{
 			"Name": "AnimationLocomotionLibraryEditor",
 			"Type": "UncookedOnly",
-			"LoadingPhase": "Default"
+			"LoadingPhase": "PreDefault"
 		}
 	]
 }
\ No newline at end of file
diff --git a/Engine/Plugins/Animation/AnimationWarping/AnimationWarping.uplugin b/Engine/Plugins/Animation/AnimationWarping/AnimationWarping.uplugin
--- a/Engine/Plugins/Animation/AnimationWarping/AnimationWarping.uplugin
+++ b/Engine/Plugins/Animation/AnimationWarping/AnimationWarping.uplugin
@@ -16,7 +16,7 @@
 		{
 			"Name": "AnimationWarpingRuntime",
 			"Type": "Runtime",
-			"LoadingPhase": "Default"
+			"LoadingPhase": "PreDefault"
 		},
 		{
 			"Name": "AnimationWarpingEditor",
